Adding experimental AES 256 GCM packet handler

diff --git a/Engine/Plugins/Experimental/PlatformCrypto/Source/PlatformCryptoOpenSSL/Private/EncryptionContextOpenSSL.cpp b/Engine/Plugins/Experimental/PlatformCrypto/Source/PlatformCryptoOpenSSL/Private/EncryptionContextOpenSSL.cpp
--- a/Engine/Plugins/Experimental/PlatformCrypto/Source/PlatformCryptoOpenSSL/Private/EncryptionContextOpenSSL.cpp
+++ b/Engine/Plugins/Experimental/PlatformCrypto/Source/PlatformCryptoOpenSSL/Private/EncryptionContextOpenSSL.cpp
@@ -4,11 +4,14 @@
 
 #include <openssl/evp.h>
 #include <openssl/rsa.h>
+#include <openssl/rand.h>
 
 DEFINE_LOG_CATEGORY_STATIC(LogPlatformCryptoOpenSSL, Warning, All);
 
 static const int32 AES256_KeySizeInBytes = 32;
 static const int32 AES256_BlockSizeInBytes = 16;
+static const int32 AES256_IVSizeInBytes = 12;
+static const int32 AES256_AuthTagSizeInBytes = 16;
 
 class FScopedEVPContext
 {
@@ -131,6 +134,157 @@ TArray<uint8> FEncryptionContextOpenSSL::Decrypt_AES_256_ECB(const TArrayView<co
 	return Plaintext;
 }
 
+TArray<uint8> FEncryptionContextOpenSSL::Encrypt_AES_256_GCM(const TArrayView<const uint8> Plaintext, const TArrayView<const uint8> Key, const TArrayView<const uint8> IV, TArray<uint8>& OutAuthTag, EPlatformCryptoResult& OutResult)
+{
+	DECLARE_SCOPE_CYCLE_COUNTER(TEXT("OpenSSL AES256GCM Encrypt"), STAT_OpenSSL_AES_GCM_Encrypt, STATGROUP_PlatformCrypto);
+
+	OutResult = EPlatformCryptoResult::Failure;
+
+	if (Key.Num() != AES256_KeySizeInBytes)
+	{
+		UE_LOG(LogPlatformCryptoOpenSSL, Warning, TEXT("FEncryptionContextOpenSSL::Encrypt_AES_256_GCM: Key size %d is not the expected size %d."), Key.Num(), AES256_KeySizeInBytes);
+		return TArray<uint8>();
+	}
+
+	FScopedEVPContext Context;
+	if (Context.Get() == nullptr)
+	{
+		UE_LOG(LogPlatformCryptoOpenSSL, Warning, TEXT("FEncryptionContextOpenSSL::Encrypt_AES_256_GCM: failed to create EVP context."));
+		return TArray<uint8>();
+	}
+
+	const int InitResult = EVP_EncryptInit_ex(Context.Get(), EVP_aes_256_gcm(), nullptr, Key.GetData(), IV.GetData());
+	if (InitResult != 1)
+	{
+		UE_LOG(LogPlatformCryptoOpenSSL, Warning, TEXT("FEncryptionContextOpenSSL::Encrypt_AES_256_GCM: EVP_EncryptInit_ex failed."));
+		return TArray<uint8>();
+	}
+
+	TArray<uint8> Ciphertext;
+	Ciphertext.SetNumUninitialized(Plaintext.Num() + AES256_BlockSizeInBytes); // Allow for up to a block of padding.
+
+	int OutLength = 0;
+	const int UpdateResult = EVP_EncryptUpdate(Context.Get(), Ciphertext.GetData(), &OutLength, Plaintext.GetData(), Plaintext.Num());
+	if (UpdateResult != 1)
+	{
+		UE_LOG(LogPlatformCryptoOpenSSL, Warning, TEXT("FEncryptionContextOpenSSL::Encrypt_AES_256_GCM: EVP_EncryptUpdate failed."));
+		return TArray<uint8>();
+	}
+
+	int FinalizeLength = 0;
+	const int FinalizeResult = EVP_EncryptFinal_ex(Context.Get(), Ciphertext.GetData() + OutLength, &FinalizeLength);
+	if (FinalizeResult != 1)
+	{
+		UE_LOG(LogPlatformCryptoOpenSSL, Warning, TEXT("FEncryptionContextOpenSSL::Encrypt_AES_256_GCM: EVP_EncryptFinal_ex failed."));
+		return TArray<uint8>();
+	}
+
+	OutAuthTag.Reset();
+	OutAuthTag.AddUninitialized(AES256_AuthTagSizeInBytes);
+
+	const int GetTagResult = EVP_CIPHER_CTX_ctrl(Context.Get(), EVP_CTRL_GCM_GET_TAG, OutAuthTag.Num(), OutAuthTag.GetData());
+	if (GetTagResult != 1)
+	{
+		UE_LOG(LogPlatformCryptoOpenSSL, Warning, TEXT("FEncryptionContextOpenSSL::Encrypt_AES_256_GCM: EVP_CIPHER_CTX_ctrl failed."));
+		return TArray<uint8>();
+	}
+
+	Ciphertext.SetNum(OutLength + FinalizeLength);
+
+	OutResult = EPlatformCryptoResult::Success;
+	return Ciphertext;
+}
+
+TArray<uint8> FEncryptionContextOpenSSL::Decrypt_AES_256_GCM(const TArrayView<const uint8> Ciphertext, const TArrayView<const uint8> Key, const TArrayView<const uint8> IV, const TArrayView<const uint8> AuthTag, EPlatformCryptoResult& OutResult)
+{
+	DECLARE_SCOPE_CYCLE_COUNTER(TEXT("OpenSSL AES256GCM Decrypt"), STAT_OpenSSL_AES_GCM_Decrypt, STATGROUP_PlatformCrypto);
+
+	OutResult = EPlatformCryptoResult::Failure;
+
+	if (Key.Num() != AES256_KeySizeInBytes)
+	{
+		UE_LOG(LogPlatformCryptoOpenSSL, Warning, TEXT("FEncryptionContextOpenSSL::Decrypt_AES_256_GCM: Key size %d is not the expected size %d."), Key.Num(), AES256_KeySizeInBytes);
+		return TArray<uint8>();
+	}
+
+	if (IV.Num() != AES256_IVSizeInBytes)
+	{
+		UE_LOG(LogPlatformCryptoOpenSSL, Warning, TEXT("FEncryptionContextOpenSSL::Decrypt_AES_256_GCM: IV size %d is not the expected size %d."), IV.Num(), AES256_IVSizeInBytes);
+		return TArray<uint8>();
+	}
+
+	if (AuthTag.Num() != AES256_AuthTagSizeInBytes)
+	{
+		UE_LOG(LogPlatformCryptoOpenSSL, Warning, TEXT("FEncryptionContextOpenSSL::Decrypt_AES_256_GCM: Auth tag size %d is not the expected size %d."), IV.Num(), AES256_AuthTagSizeInBytes);
+		return TArray<uint8>();
+	}
+
+	FScopedEVPContext Context;
+	if (Context.Get() == nullptr)
+	{
+		UE_LOG(LogPlatformCryptoOpenSSL, Warning, TEXT("FEncryptionContextOpenSSL::Decrypt_AES_256_GCM: failed to create EVP context."));
+		return TArray<uint8>();
+	}
+
+	const int InitResult = EVP_DecryptInit_ex(Context.Get(), EVP_aes_256_gcm(), nullptr, Key.GetData(), IV.GetData());
+	if (InitResult != 1)
+	{
+		UE_LOG(LogPlatformCryptoOpenSSL, Warning, TEXT("FEncryptionContextOpenSSL::Decrypt_AES_256_GCM: EVP_DecryptInit_ex failed."));
+		return TArray<uint8>();
+	}
+
+	TArray<uint8> Plaintext;
+	Plaintext.SetNumUninitialized(Ciphertext.Num());
+
+	int OutLength = 0;
+	const int UpdateResult = EVP_DecryptUpdate(Context.Get(), Plaintext.GetData(), &OutLength, Ciphertext.GetData(), Ciphertext.Num());
+	if (UpdateResult != 1)
+	{
+		UE_LOG(LogPlatformCryptoOpenSSL, Warning, TEXT("FEncryptionContextOpenSSL::Decrypt_AES_256_GCM: EVP_DecryptUpdate failed."));
+		return TArray<uint8>();
+	}
+
+	const int SetTagResult = EVP_CIPHER_CTX_ctrl(Context.Get(), EVP_CTRL_GCM_SET_TAG, AuthTag.Num(), const_cast<uint8*>(AuthTag.GetData()));
+	if (SetTagResult != 1)
+	{
+		UE_LOG(LogPlatformCryptoOpenSSL, Warning, TEXT("FEncryptionContextOpenSSL::Decrypt_AES_256_GCM: EVP_CIPHER_CTX_ctrl failed."));
+		return TArray<uint8>();
+	}
+
+	int FinalizeLength = 0;
+	const int FinalizeResult = EVP_DecryptFinal_ex(Context.Get(), Plaintext.GetData() + OutLength, &FinalizeLength);
+	if (FinalizeResult != 1)
+	{
+		UE_LOG(LogPlatformCryptoOpenSSL, Warning, TEXT("FEncryptionContextOpenSSL::Decrypt_AES_256_GCM: EVP_DecryptFinal_ex failed."));
+		return TArray<uint8>();
+	}
+
+	Plaintext.SetNum(OutLength + FinalizeLength);
+
+	OutResult = EPlatformCryptoResult::Success;
+	return Plaintext;
+}
+
+TArray<uint8> FEncryptionContextOpenSSL::GetRandomBytes(uint32 NumBytes, EPlatformCryptoResult& OutResult)
+{
+	DECLARE_SCOPE_CYCLE_COUNTER(TEXT("OpenSSL GetRandomBytes"), STAT_OpenSSL_GetRandomBytes, STATGROUP_PlatformCrypto);
+
+	OutResult = EPlatformCryptoResult::Failure;
+
+	TArray<uint8> RandomBytes;
+	RandomBytes.AddUninitialized(NumBytes);
+
+	const int RandResult = RAND_bytes(RandomBytes.GetData(), RandomBytes.Num());
+	if (RandResult != 1)
+	{
+		UE_LOG(LogPlatformCryptoOpenSSL, Warning, TEXT("FEncryptionContextOpenSSL::GetRandomBytes: RAND_bytes failed."));
+		return TArray<uint8>();
+	}
+	
+	OutResult = EPlatformCryptoResult::Success;
+	return RandomBytes;
+}
+
 class FScopedEVPMDContext
 {
 public:
diff --git a/Engine/Plugins/Experimental/PlatformCrypto/Source/PlatformCryptoOpenSSL/Public/EncryptionContextOpenSSL.h b/Engine/Plugins/Experimental/PlatformCrypto/Source/PlatformCryptoOpenSSL/Public/EncryptionContextOpenSSL.h
--- a/Engine/Plugins/Experimental/PlatformCrypto/Source/PlatformCryptoOpenSSL/Public/EncryptionContextOpenSSL.h
+++ b/Engine/Plugins/Experimental/PlatformCrypto/Source/PlatformCryptoOpenSSL/Public/EncryptionContextOpenSSL.h
@@ -15,9 +15,13 @@ class PLATFORMCRYPTOOPENSSL_API FEncryptionContextOpenSSL
 public:
 
 	TArray<uint8> Encrypt_AES_256_ECB(const TArrayView<const uint8> Plaintext, const TArrayView<const uint8> Key, EPlatformCryptoResult& OutResult);
-
 	TArray<uint8> Decrypt_AES_256_ECB(const TArrayView<const uint8> Ciphertext, const TArrayView<const uint8> Key, EPlatformCryptoResult& OutResult);
 
+	TArray<uint8> Encrypt_AES_256_GCM(const TArrayView<const uint8> Plaintext, const TArrayView<const uint8> Key, const TArrayView<const uint8> IV, TArray<uint8>& OutAuthTag, EPlatformCryptoResult& OutResult);
+	TArray<uint8> Decrypt_AES_256_GCM(const TArrayView<const uint8> Ciphertext, const TArrayView<const uint8> Key, const TArrayView<const uint8> IV, const TArrayView<const uint8> AuthTag, EPlatformCryptoResult& OutResult);
+
+	TArray<uint8> GetRandomBytes(uint32 NumBytes, EPlatformCryptoResult& OutResult);
+
 	bool DigestVerify_PS256(const TArrayView<const char> Message, const TArrayView<const uint8> Signature, const TArrayView<const uint8> PKCS1Key);
 
 	TPlatformCryptoRSAKey CreateKey_RSA(const TArrayView<const uint8> PublicExponent, const TArrayView<const uint8> PrivateExponent, const TArrayView<const uint8> Modulus);
diff --git a/Engine/Plugins/Runtime/PacketHandlers/AESGCMHandlerComponent/AESGCMHandlerComponent.uplugin b/Engine/Plugins/Runtime/PacketHandlers/AESGCMHandlerComponent/AESGCMHandlerComponent.uplugin
new file mode 100644
--- /dev/null
+++ b/Engine/Plugins/Runtime/PacketHandlers/AESGCMHandlerComponent/AESGCMHandlerComponent.uplugin
@@ -0,0 +1,56 @@
+{
+	"FileVersion" : 3,
+	"Version" : 1,
+	"VersionName" : "1.0",
+	"FriendlyName" : "AES GCM network packet handler",
+	"Description" : "Provides a packet handler component to do AES GCM encryption and decryption.",
+	"Category" : "Misc",
+	"CreatedBy" : "Epic Games, Inc.",
+	"CreatedByURL" : "http://epicgames.com",
+	"DocsURL" : "",
+	"MarketplaceURL" : "",
+	"SupportURL" : "",
+	"EnabledByDefault" : false,
+	"CanContainContent" : false,
+	"IsBetaVersion" : true,
+	"Installed" : false,
+	"Modules" :
+	[
+		{
+			"Name" : "AESGCMHandlerComponent",
+			"Type" : "Runtime",
+			"LoadingPhase" : "Default",
+			"WhitelistPlatforms":
+			[
+				"Android",
+				"IOS",
+				"Mac",
+				"Win32",
+				"Win64",
+				"Linux",
+				"PS4",
+				"XboxOne",
+				"Switch" 
+			]
+		}
+	],
+	"Plugins":
+	[
+		{
+			"Name": "PlatformCrypto",
+			"Enabled": true,
+			"WhitelistPlatforms":
+			[
+				"Android",
+				"IOS",
+				"Mac",
+				"Win32",
+				"Win64",
+				"Linux",
+				"PS4",
+				"XboxOne",
+				"Switch" 
+			]
+		}
+	]
+}
\ No newline at end of file
diff --git a/Engine/Plugins/Runtime/PacketHandlers/AESGCMHandlerComponent/Source/AESGCMHandlerComponent.Build.cs b/Engine/Plugins/Runtime/PacketHandlers/AESGCMHandlerComponent/Source/AESGCMHandlerComponent.Build.cs
new file mode 100644
--- /dev/null
+++ b/Engine/Plugins/Runtime/PacketHandlers/AESGCMHandlerComponent/Source/AESGCMHandlerComponent.Build.cs
@@ -0,0 +1,59 @@
+// Copyright 1998-2019 Epic Games, Inc. All Rights Reserved.
+
+using UnrealBuildTool;
+using System.IO;
+
+public class AESGCMHandlerComponent : ModuleRules
+{
+    public AESGCMHandlerComponent(ReadOnlyTargetRules Target) : base(Target)
+	{
+		PrivateIncludePaths.AddRange(
+			new string[] {
+			}
+			);
+
+		PublicDependencyModuleNames.AddRange(
+			new string[]
+			{
+				"Core",
+				"PacketHandler",
+				"PlatformCrypto",
+			}
+			);
+
+		PublicIncludePathModuleNames.AddRange(
+			new string[]
+			{
+				"PlatformCrypto"
+			}
+			);
+
+		if (Target.Platform == UnrealTargetPlatform.XboxOne)
+		{
+			PublicDependencyModuleNames.AddRange(
+				new string[]
+				{
+					"PlatformCryptoBCrypt",
+				}
+				);
+		}
+		else if (Target.Platform == UnrealTargetPlatform.Switch)
+		{
+			PublicDependencyModuleNames.AddRange(
+				new string[]
+				{
+					"PlatformCryptoSwitch",
+				}
+				);
+		}
+		else
+		{
+			PublicDependencyModuleNames.AddRange(
+				new string[]
+				{
+					"PlatformCryptoOpenSSL",
+				}
+				);
+		}
+	}
+}
\ No newline at end of file
diff --git a/Engine/Plugins/Runtime/PacketHandlers/AESGCMHandlerComponent/Source/Private/AESGCMHandlerComponent.cpp b/Engine/Plugins/Runtime/PacketHandlers/AESGCMHandlerComponent/Source/Private/AESGCMHandlerComponent.cpp
new file mode 100644
--- /dev/null
+++ b/Engine/Plugins/Runtime/PacketHandlers/AESGCMHandlerComponent/Source/Private/AESGCMHandlerComponent.cpp
@@ -0,0 +1,240 @@
+// Copyright 1998-2019 Epic Games, Inc. All Rights Reserved.
+
+#include "AESGCMHandlerComponent.h"
+
+IMPLEMENT_MODULE( FAESGCMHandlerComponentModule, AESGCMHandlerComponent )
+
+TSharedPtr<HandlerComponent> FAESGCMHandlerComponentModule::CreateComponentInstance(FString& Options)
+{
+	return MakeShared<FAESGCMHandlerComponent>();
+}
+
+
+FAESGCMHandlerComponent::FAESGCMHandlerComponent()
+	: FEncryptionComponent(FName(TEXT("AESGCMHandlerComponent")))
+	, bEncryptionEnabled(false)
+{
+	EncryptionContext = IPlatformCrypto::Get().CreateContext();
+}
+
+void FAESGCMHandlerComponent::SetEncryptionKey(TArrayView<const uint8> NewKey)
+{
+	if (NewKey.Num() != KeySizeInBytes)
+	{
+		UE_LOG(PacketHandlerLog, Log, TEXT("FAESGCMHandlerComponent::SetEncryptionKey. NewKey is not %d bytes long, ignoring."), KeySizeInBytes);
+		return;
+	}
+
+	Key.Reset(KeySizeInBytes);
+	Key.Append(NewKey.GetData(), NewKey.Num());
+}
+
+void FAESGCMHandlerComponent::EnableEncryption()
+{
+	bEncryptionEnabled = true;
+}
+
+void FAESGCMHandlerComponent::DisableEncryption()
+{
+	bEncryptionEnabled = false;
+}
+
+bool FAESGCMHandlerComponent::IsEncryptionEnabled() const
+{
+	return bEncryptionEnabled;
+}
+
+void FAESGCMHandlerComponent::Initialize()
+{
+	SetActive(true);
+	SetState(Handler::Component::State::Initialized);
+	Initialized();
+}
+
+bool FAESGCMHandlerComponent::IsValid() const
+{
+	return true;
+}
+
+void FAESGCMHandlerComponent::Incoming(FBitReader& Packet)
+{
+	DECLARE_SCOPE_CYCLE_COUNTER(TEXT("PacketHandler AESGCM Decrypt"), STAT_PacketHandler_AESGCM_Decrypt, STATGROUP_Net);
+
+	// Handle this packet
+	if (IsValid() && Packet.GetNumBytes() > 0)
+	{
+		// Check first bit to see whether payload is encrypted
+		if (Packet.ReadBit() != 0)
+		{
+			// If the key hasn't been set yet, we can't decrypt, so ignore this packet. We don't set an error in this case because it may just be an out-of-order packet.
+			if (Key.Num() == 0)
+			{
+				UE_LOG(PacketHandlerLog, Log, TEXT("FAESGCMHandlerComponent::Incoming: received encrypted packet before key was set, ignoring."));
+				Packet.SetData(nullptr, 0);
+				return;
+			}
+
+			IV.Reset();
+			IV.AddUninitialized(IVSizeInBytes);
+			Packet.SerializeBits(IV.GetData(), IV.Num() * 8);
+
+			AuthTag.Reset();
+			AuthTag.AddUninitialized(AuthTagSizeInBytes);
+			Packet.SerializeBits(AuthTag.GetData(), AuthTag.Num() * 8);
+
+			// Copy remaining bits to a TArray so that they are byte-aligned.
+			Ciphertext.Reset();
+			Ciphertext.AddUninitialized(Packet.GetBytesLeft());
+			Ciphertext[Ciphertext.Num() - 1] = 0;
+
+			Packet.SerializeBits(Ciphertext.GetData(), Packet.GetBitsLeft());
+
+			UE_LOG(PacketHandlerLog, VeryVerbose, TEXT("AESGCM packet handler received %ld bytes before decryption."), Ciphertext.Num());
+
+			EPlatformCryptoResult DecryptResult = EPlatformCryptoResult::Failure;
+			TArray<uint8> Plaintext = EncryptionContext->Decrypt_AES_256_GCM(Ciphertext, Key, IV, AuthTag, DecryptResult);
+
+			if (DecryptResult == EPlatformCryptoResult::Failure)
+			{
+				UE_LOG(PacketHandlerLog, Log, TEXT("FAESGCMHandlerComponent::Incoming: failed to decrypt packet."));
+				Packet.SetError();
+				return;
+			}
+
+			if (Plaintext.Num() == 0)
+			{
+				Packet.SetData(nullptr, 0);
+				return;
+			}
+
+			// Look for the termination bit that was written in Outgoing() to determine the exact bit size.
+			uint8 LastByte = Plaintext.Last();
+
+			if (LastByte != 0)
+			{
+				int32 BitSize = (Plaintext.Num() * 8) - 1;
+
+				// Bit streaming, starts at the Least Significant Bit, and ends at the MSB.
+				while (!(LastByte & 0x80))
+				{
+					LastByte *= 2;
+					BitSize--;
+				}
+
+				UE_LOG(PacketHandlerLog, VeryVerbose, TEXT("  Have %d bits after decryption."), BitSize);
+
+				Packet.SetData(MoveTemp(Plaintext), BitSize);
+			}
+			else
+			{
+				UE_LOG(PacketHandlerLog, Log, TEXT("FAESGCMHandlerComponent::Incoming: malformed packet, last byte was 0."));
+				Packet.SetError();
+			}
+		}
+	}
+}
+
+void FAESGCMHandlerComponent::Outgoing(FBitWriter& Packet, FOutPacketTraits& Traits)
+{
+	DECLARE_SCOPE_CYCLE_COUNTER(TEXT("PacketHandler AESGCM Encrypt"), STAT_PacketHandler_AESGCM_Encrypt, STATGROUP_Net);
+
+	// Handle this packet
+	if (IsValid() && Packet.GetNumBytes() > 0)
+	{
+		// Allow for encryption enabled bit and termination bit. Allow resizing to account for encryption padding.
+		FBitWriter NewPacket(Packet.GetNumBits() + 2, true);
+		NewPacket.WriteBit(bEncryptionEnabled ? 1 : 0);
+
+		if (NewPacket.IsError())
+		{
+			UE_LOG(PacketHandlerLog, Log, TEXT("FAESGCMHandlerComponent::Outgoing: failed to write encryption bit."));
+			Packet.SetError();
+			return;
+		}
+
+		if (bEncryptionEnabled)
+		{
+			UE_LOG(PacketHandlerLog, VeryVerbose, TEXT("AESGCM packet handler sending %ld bits before encryption."), Packet.GetNumBits());
+
+			// Write a termination bit so that the receiving side can calculate the exact number of bits sent.
+			// Same technique used in UNetConnection.
+			Packet.WriteBit(1);
+
+			if (Packet.IsError())
+			{
+				UE_LOG(PacketHandlerLog, Log, TEXT("FAESGCMHandlerComponent::Outgoing: failed to write termination bit."));
+				return;
+			}
+
+			EPlatformCryptoResult RandResult = EPlatformCryptoResult::Failure;
+			TArray<uint8> OutIV = EncryptionContext->GetRandomBytes(IVSizeInBytes, RandResult);
+			if (RandResult == EPlatformCryptoResult::Failure)
+			{
+				UE_LOG(PacketHandlerLog, Log, TEXT("FAESGCMHandlerComponent::Outgoing: failed to generate IV."));
+				Packet.SetError();
+				return;
+			}
+
+			TArray<uint8> OutAuthTag;
+
+			EPlatformCryptoResult EncryptResult = EPlatformCryptoResult::Failure;
+			TArray<uint8> OutCiphertext = EncryptionContext->Encrypt_AES_256_GCM(TArrayView<uint8>(Packet.GetData(), Packet.GetNumBytes()), Key, OutIV, OutAuthTag, EncryptResult);
+
+			if (EncryptResult == EPlatformCryptoResult::Failure)
+			{
+				UE_LOG(PacketHandlerLog, Log, TEXT("FAESGCMHandlerComponent::Outgoing: failed to encrypt packet."));
+				Packet.SetError();
+				return;
+			}
+			else
+			{
+				NewPacket.Serialize(OutIV.GetData(), OutIV.Num());
+				NewPacket.Serialize(OutAuthTag.GetData(), OutAuthTag.Num());
+				NewPacket.Serialize(OutCiphertext.GetData(), OutCiphertext.Num());
+
+				if (NewPacket.IsError())
+				{
+					UE_LOG(PacketHandlerLog, Log, TEXT("FAESGCMHandlerComponent::Outgoing: failed to write ciphertext to packet."));
+					Packet.SetError();
+					return;
+				}
+
+				UE_LOG(PacketHandlerLog, VeryVerbose, TEXT("  AESGCM packet handler sending %d bytes after encryption."), NewPacket.GetNumBytes());
+			}
+		}
+		else
+		{
+			NewPacket.SerializeBits(Packet.GetData(), Packet.GetNumBits());
+		}
+
+		Packet = MoveTemp(NewPacket);
+	}
+}
+
+int32 FAESGCMHandlerComponent::GetReservedPacketBits() const
+{
+	// Worst case includes the encryption enabled bit, the termination bit, padding up to the next whole byte, and a block of padding.
+	// Now also includes initialization vector and auth tag
+	return 2 + 7 + ((BlockSizeInBytes + IVSizeInBytes + AuthTagSizeInBytes) * 8);
+}
+
+void FAESGCMHandlerComponent::CountBytes(FArchive& Ar) const
+{
+	FEncryptionComponent::CountBytes(Ar);
+
+	const SIZE_T SizeOfThis = sizeof(*this) - sizeof(FEncryptionComponent);
+	Ar.CountBytes(SizeOfThis, SizeOfThis);
+
+	/*
+	Note, as of now, EncryptionContext is just typedef'd, but none of the base
+	types actually allocated memory directly in their classes (although there may be
+	global state).
+	if (FEncryptionContext const * const LocalContext = EncrpytionContext.Get())
+	{
+		LocalContext->CountBytes(Ar);
+	}
+	*/
+
+	Key.CountBytes(Ar);
+	Ciphertext.CountBytes(Ar);
+}
diff --git a/Engine/Plugins/Runtime/PacketHandlers/AESGCMHandlerComponent/Source/Public/AESGCMHandlerComponent.h b/Engine/Plugins/Runtime/PacketHandlers/AESGCMHandlerComponent/Source/Public/AESGCMHandlerComponent.h
new file mode 100644
--- /dev/null
+++ b/Engine/Plugins/Runtime/PacketHandlers/AESGCMHandlerComponent/Source/Public/AESGCMHandlerComponent.h
@@ -0,0 +1,76 @@
+// Copyright 1998-2019 Epic Games, Inc. All Rights Reserved.
+
+#pragma once
+
+#include "CoreMinimal.h"
+#include "PacketHandler.h"
+#include "IPlatformCrypto.h"
+#include "EncryptionComponent.h"
+
+/*
+* AES256 GCM block encryption component.
+*/
+class AESGCMHANDLERCOMPONENT_API FAESGCMHandlerComponent : public FEncryptionComponent
+{
+public:
+	/**
+	 * Default constructor that leaves the Key empty, and encryption disabled.
+	 * You must set the key before enabling encryption, or before receiving encrypted
+	 * packets, or those operations will fail.
+	 */
+	FAESGCMHandlerComponent();
+
+	// This handler uses AES256, which has 32-byte keys.
+	static const int32 KeySizeInBytes = 32;
+
+	// This handler uses AES256, which has 32-byte keys.
+	static const int32 BlockSizeInBytes = 16;
+
+	static const int32 IVSizeInBytes = 12;
+	static const int32 AuthTagSizeInBytes = 16;
+
+	// Replace the key used for encryption with NewKey if NewKey is exactly KeySizeInBytes long.
+	virtual void SetEncryptionKey(TArrayView<const uint8> NewKey) override;
+
+	// After calling this, future outgoing packets will be encrypted (until a call to DisableEncryption).
+	virtual void EnableEncryption() override;
+
+	// After calling this, future outgoing packets will not be encrypted (until a call to DisableEncryption).
+	virtual void DisableEncryption() override;
+
+	// Returns true if encryption is currently enabled.
+	virtual bool IsEncryptionEnabled() const override;
+
+	// HandlerComponent interface
+	virtual void Initialize() override;
+	virtual bool IsValid() const override;
+	virtual void Incoming(FBitReader& Packet) override;
+	virtual void Outgoing(FBitWriter& Packet, FOutPacketTraits& Traits) override;
+	virtual void IncomingConnectionless(const TSharedPtr<const FInternetAddr>& Address, FBitReader& Packet) override {}
+	virtual void OutgoingConnectionless(const TSharedPtr<const FInternetAddr>& Address, FBitWriter& Packet, FOutPacketTraits& Traits) override {}
+	virtual int32 GetReservedPacketBits() const override;
+	virtual void CountBytes(FArchive& Ar) const override;
+
+private:
+	TUniquePtr<FEncryptionContext> EncryptionContext;
+
+	TArray<uint8> Key;
+	
+	// Avoid per packet allocations
+	TArray<uint8> Ciphertext;
+	TArray<uint8> IV;
+	TArray<uint8> AuthTag;
+
+	bool bEncryptionEnabled;
+};
+
+
+/**
+ * The public interface to this module.
+ */
+class FAESGCMHandlerComponentModule : public FPacketHandlerComponentModuleInterface
+{
+public:
+	/* Creates an instance of this component */
+	virtual TSharedPtr<HandlerComponent> CreateComponentInstance(FString& Options) override;
+};
