Fix potential underflow errors

diff --git a/Engine/Plugins/Runtime/PacketHandlers/AESGCMHandlerComponent/Source/Private/AESGCMHandlerComponent.cpp b/Engine/Plugins/Runtime/PacketHandlers/AESGCMHandlerComponent/Source/Private/AESGCMHandlerComponent.cpp
--- a/Engine/Plugins/Runtime/PacketHandlers/AESGCMHandlerComponent/Source/Private/AESGCMHandlerComponent.cpp
+++ b/Engine/Plugins/Runtime/PacketHandlers/AESGCMHandlerComponent/Source/Private/AESGCMHandlerComponent.cpp
@@ -74,20 +74,47 @@ void FAESGCMHandlerComponent::Incoming(FBitReader& Packet)
 				return;
 			}
 
-			IV.Reset();
-			IV.AddUninitialized(IVSizeInBytes);
-			Packet.SerializeBits(IV.GetData(), IV.Num() * 8);
+			if (Packet.GetBytesLeft() >= IVSizeInBytes)
+			{
+				IV.Reset();
+				IV.AddUninitialized(IVSizeInBytes);
+				Packet.SerializeBits(IV.GetData(), IV.Num() * 8);
+			}
+			else
+			{
+				UE_LOG(PacketHandlerLog, Log, TEXT("FAESGCMHandlerComponent::Incoming: missing IV"));
+				Packet.SetError();
+				return;
+			}
 
-			AuthTag.Reset();
-			AuthTag.AddUninitialized(AuthTagSizeInBytes);
-			Packet.SerializeBits(AuthTag.GetData(), AuthTag.Num() * 8);
+			if (Packet.GetBytesLeft() >= AuthTagSizeInBytes)
+			{
+				AuthTag.Reset();
+				AuthTag.AddUninitialized(AuthTagSizeInBytes);
+				Packet.SerializeBits(AuthTag.GetData(), AuthTag.Num() * 8);
+			}
+			else
+			{
+				UE_LOG(PacketHandlerLog, Log, TEXT("FAESGCMHandlerComponent::Incoming: missing auth tag"));
+				Packet.SetError();
+				return;
+			}
 
 			// Copy remaining bits to a TArray so that they are byte-aligned.
-			Ciphertext.Reset();
-			Ciphertext.AddUninitialized(Packet.GetBytesLeft());
-			Ciphertext[Ciphertext.Num() - 1] = 0;
+			if (Packet.GetBytesLeft() > 0)
+			{
+				Ciphertext.Reset();
+				Ciphertext.AddUninitialized(Packet.GetBytesLeft());
+				Ciphertext[Ciphertext.Num() - 1] = 0;
 
-			Packet.SerializeBits(Ciphertext.GetData(), Packet.GetBitsLeft());
+				Packet.SerializeBits(Ciphertext.GetData(), Packet.GetBitsLeft());
+			}
+			else
+			{
+				UE_LOG(PacketHandlerLog, Log, TEXT("FAESGCMHandlerComponent::Incoming: missing ciphertext"));
+				Packet.SetError();
+				return;
+			}
 
 			UE_LOG(PacketHandlerLog, VeryVerbose, TEXT("AESGCM packet handler received %ld bytes before decryption."), Ciphertext.Num());
 
diff --git a/Engine/Plugins/Runtime/PacketHandlers/AESHandlerComponent/Source/Private/AESHandlerComponent.cpp b/Engine/Plugins/Runtime/PacketHandlers/AESHandlerComponent/Source/Private/AESHandlerComponent.cpp
--- a/Engine/Plugins/Runtime/PacketHandlers/AESHandlerComponent/Source/Private/AESHandlerComponent.cpp
+++ b/Engine/Plugins/Runtime/PacketHandlers/AESHandlerComponent/Source/Private/AESHandlerComponent.cpp
@@ -79,9 +79,18 @@ void FAESHandlerComponent::Incoming(FBitReader& Packet)
 			}
 
 			// Copy remaining bits to a TArray so that they are byte-aligned.
-			Ciphertext.Reset();
-			Ciphertext.AddUninitialized(Packet.GetBytesLeft());
-			Ciphertext[Ciphertext.Num() - 1] = 0;
+			if (Packet.GetBytesLeft() > 0)
+			{
+				Ciphertext.Reset();
+				Ciphertext.AddUninitialized(Packet.GetBytesLeft());
+				Ciphertext[Ciphertext.Num() - 1] = 0;
+			}
+			else
+			{
+				UE_LOG(PacketHandlerLog, Log, TEXT("FAESHandlerComponent::Incoming: missing ciphertext"));
+				Packet.SetError();
+				return;
+			}
 
 			Packet.SerializeBits(Ciphertext.GetData(), Packet.GetBitsLeft());
 
