[Switch] Final platform extension plugin changes.

diff --git a/Engine/Plugins/Experimental/PlatformCrypto/PlatformCrypto.uplugin b/Engine/Plugins/Experimental/PlatformCrypto/PlatformCrypto.uplugin
--- a/Engine/Plugins/Experimental/PlatformCrypto/PlatformCrypto.uplugin
+++ b/Engine/Plugins/Experimental/PlatformCrypto/PlatformCrypto.uplugin
@@ -28,8 +28,7 @@
 				"Win64",
 				"Linux",
 				"PS4",
-				"XboxOne",
-				"Switch"
+				"XboxOne"
 			],
 			"BlacklistPlatforms": [
 				"TVOS"
@@ -48,8 +47,7 @@
 				"Win64",
 				"Linux",
 				"PS4",
-				"XboxOne",
-				"Switch"
+				"XboxOne"
 			],
 			"BlacklistPlatforms": [
 				"TVOS"
@@ -84,19 +82,6 @@
 				"IOS",
 				"TVOS"
 			]
-		},
-		{
-			"Name": "PlatformCryptoSwitch",
-			"Type": "RuntimeAndProgram",
-			"LoadingPhase": "EarliestPossible",
-			"WhitelistPlatforms": [
-				"Switch"
-			],
-			"BlacklistPlatforms": [
-				"Android",
-				"IOS",
-				"TVOS"
-			]
 		}
 	]
 }
\ No newline at end of file
diff --git a/Engine/Plugins/Experimental/PlatformCrypto/Source/PlatformCrypto/PlatformCrypto.Build.cs b/Engine/Plugins/Experimental/PlatformCrypto/Source/PlatformCrypto/PlatformCrypto.Build.cs
--- a/Engine/Plugins/Experimental/PlatformCrypto/Source/PlatformCrypto/PlatformCrypto.Build.cs
+++ b/Engine/Plugins/Experimental/PlatformCrypto/Source/PlatformCrypto/PlatformCrypto.Build.cs
@@ -16,16 +16,7 @@ namespace UnrealBuildTool.Rules
 				}
 				);
 
-			if (Target.Platform == UnrealTargetPlatform.Switch)
-			{
-				PublicDependencyModuleNames.AddRange(
-					new string[]
-					{
-						"PlatformCryptoSwitch",
-					}
-					);
-			}
-			else if (DefaultToSSL)
+			if (DefaultToSSL)
 			{
 				PublicDependencyModuleNames.AddRange(
 					new string[]
diff --git a/Engine/Plugins/Runtime/PacketHandlers/AESGCMHandlerComponent/AESGCMHandlerComponent.uplugin b/Engine/Plugins/Runtime/PacketHandlers/AESGCMHandlerComponent/AESGCMHandlerComponent.uplugin
--- a/Engine/Plugins/Runtime/PacketHandlers/AESGCMHandlerComponent/AESGCMHandlerComponent.uplugin
+++ b/Engine/Plugins/Runtime/PacketHandlers/AESGCMHandlerComponent/AESGCMHandlerComponent.uplugin
@@ -29,8 +29,7 @@
 				"Win64",
 				"Linux",
 				"PS4",
-				"XboxOne",
-				"Switch" 
+				"XboxOne"
 			]
 		}
 	],
@@ -48,8 +47,7 @@
 				"Win64",
 				"Linux",
 				"PS4",
-				"XboxOne",
-				"Switch" 
+				"XboxOne"
 			]
 		}
 	]
diff --git a/Engine/Plugins/Runtime/PacketHandlers/AESGCMHandlerComponent/Source/AESGCMHandlerComponent.Build.cs b/Engine/Plugins/Runtime/PacketHandlers/AESGCMHandlerComponent/Source/AESGCMHandlerComponent.Build.cs
--- a/Engine/Plugins/Runtime/PacketHandlers/AESGCMHandlerComponent/Source/AESGCMHandlerComponent.Build.cs
+++ b/Engine/Plugins/Runtime/PacketHandlers/AESGCMHandlerComponent/Source/AESGCMHandlerComponent.Build.cs
@@ -5,7 +5,9 @@ using System.IO;
 
 public class AESGCMHandlerComponent : ModuleRules
 {
-    public AESGCMHandlerComponent(ReadOnlyTargetRules Target) : base(Target)
+	protected virtual bool DefaultToSSL { get { return true; } }
+
+	public AESGCMHandlerComponent(ReadOnlyTargetRules Target) : base(Target)
 	{
 		PrivateIncludePaths.AddRange(
 			new string[] {
@@ -37,16 +39,7 @@ public class AESGCMHandlerComponent : ModuleRules
 				}
 				);
 		}
-		else if (Target.Platform == UnrealTargetPlatform.Switch)
-		{
-			PublicDependencyModuleNames.AddRange(
-				new string[]
-				{
-					"PlatformCryptoSwitch",
-				}
-				);
-		}
-		else
+		else if (DefaultToSSL)
 		{
 			PublicDependencyModuleNames.AddRange(
 				new string[]
diff --git a/Engine/Plugins/Runtime/PacketHandlers/AESHandlerComponent/AESHandlerComponent.uplugin b/Engine/Plugins/Runtime/PacketHandlers/AESHandlerComponent/AESHandlerComponent.uplugin
--- a/Engine/Plugins/Runtime/PacketHandlers/AESHandlerComponent/AESHandlerComponent.uplugin
+++ b/Engine/Plugins/Runtime/PacketHandlers/AESHandlerComponent/AESHandlerComponent.uplugin
@@ -29,8 +29,7 @@
 				"Win64",
 				"Linux",
 				"PS4",
-				"XboxOne",
-				"Switch" 
+				"XboxOne" 
 			]
 		}
 	],
@@ -48,8 +47,7 @@
 				"Win64",
 				"Linux",
 				"PS4",
-				"XboxOne",
-				"Switch" 
+				"XboxOne" 
 			]
 		}
 	]
diff --git a/Engine/Plugins/Runtime/PacketHandlers/AESHandlerComponent/Source/AESHandlerComponent.Build.cs b/Engine/Plugins/Runtime/PacketHandlers/AESHandlerComponent/Source/AESHandlerComponent.Build.cs
--- a/Engine/Plugins/Runtime/PacketHandlers/AESHandlerComponent/Source/AESHandlerComponent.Build.cs
+++ b/Engine/Plugins/Runtime/PacketHandlers/AESHandlerComponent/Source/AESHandlerComponent.Build.cs
@@ -39,15 +39,6 @@ public class AESHandlerComponent : ModuleRules
 				}
 				);
 		}
-		else if (Target.Platform == UnrealTargetPlatform.Switch)
-		{
-			PublicDependencyModuleNames.AddRange(
-				new string[]
-				{
-					"PlatformCryptoSwitch",
-				}
-				);
-		}
 		else if (DefaultToSSL)
 		{
 			PublicDependencyModuleNames.AddRange(
