Convert remaining Net Fault Handler enum LexToString's to FOREACH_ENUM_* format.

diff --git a/Engine/Plugins/Compression/OodleNetwork/Source/Private/OodleNetworkFaultHandler.cpp b/Engine/Plugins/Compression/OodleNetwork/Source/Private/OodleNetworkFaultHandler.cpp
--- a/Engine/Plugins/Compression/OodleNetwork/Source/Private/OodleNetworkFaultHandler.cpp
+++ b/Engine/Plugins/Compression/OodleNetwork/Source/Private/OodleNetworkFaultHandler.cpp
@@ -11,32 +11,27 @@
  * EOodleNetResult
  */
 
-const TCHAR* LexToString(EOodleNetResult InResult)
-{
-	switch (InResult)
-	{
-	case EOodleNetResult::Unknown:
-		return TEXT("Unknown");
-
-	case EOodleNetResult::Success:
-		return TEXT("Success");
-
-	case EOodleNetResult::OodleDecodeFailed:
-		return TEXT("OodleDecodeFailed");
+#ifndef CASE_ENUM_TO_TEXT_RET
+#define CASE_ENUM_TO_TEXT_RET(txt) case txt: ReturnVal = TEXT(#txt); break;
+#endif
 
-	case EOodleNetResult::OodleSerializePayloadFail:
-		return TEXT("OodleSerializePayloadFail");
+const TCHAR* LexToString(EOodleNetResult Enum)
+{
+	const TCHAR* ReturnVal = TEXT("::Invalid");
 
-	case EOodleNetResult::OodleBadDecompressedLength:
-		return TEXT("OodleBadDecompressedLength");
+	switch (Enum)
+	{
+		FOREACH_ENUM_EOODLENETRESULT(CASE_ENUM_TO_TEXT_RET)
+	}
 
-	case EOodleNetResult::OodleNoDictionary:
-		return TEXT("OodleNoDictionary");
+	while (*ReturnVal != ':')
+	{
+		ReturnVal++;
+	}
 
+	ReturnVal += 2;
 
-	default:
-		return TEXT("Invalid");
-	}
+	return ReturnVal;
 }
 
 
@@ -104,35 +99,3 @@ UE::Net::EHandleNetResult FOodleNetworkFaultHandler::HandleNetResult(UE::Net::FN
 		
 	return ReturnVal;
 }
-
-#if !(UE_BUILD_SHIPPING || UE_BUILD_TEST)
-#include "Misc/AutomationTest.h"
-
-IMPLEMENT_SIMPLE_AUTOMATION_TEST(FOodleNetResultEnumTest, "System.Core.Networking.EOodleNetResult.EnumTest",
-									EAutomationTestFlags::ApplicationContextMask | EAutomationTestFlags::SmokeFilter);
-
-bool FOodleNetResultEnumTest::RunTest(const FString& Parameters)
-{
-	const UEnum* OodleResultEnum = StaticEnum<EOodleNetResult>();
-
-	if (TestTrue(TEXT("EOodleNetResult must exist"), OodleResultEnum != nullptr) && OodleResultEnum != nullptr)
-	{
-		const int64 OodleResultEnumLast = OodleResultEnum->GetMaxEnumValue() - 1;
-		bool bLexMismatch = false;
-
-		for (int64 EnumIdx=0; EnumIdx<=OodleResultEnumLast; EnumIdx++)
-		{
-			if (OodleResultEnum->GetNameStringByValue(EnumIdx) != LexToString((EOodleNetResult)EnumIdx))
-			{
-				bLexMismatch = true;
-				break;
-			}
-		}
-
-		TestFalse(TEXT("EOodleNetResult must not be missing LexToString entries"), bLexMismatch);
-	}
-
-	return true;
-}
-#endif
-
diff --git a/Engine/Plugins/Compression/OodleNetwork/Source/Public/OodleNetworkFaultHandler.h b/Engine/Plugins/Compression/OodleNetwork/Source/Public/OodleNetworkFaultHandler.h
--- a/Engine/Plugins/Compression/OodleNetwork/Source/Public/OodleNetworkFaultHandler.h
+++ b/Engine/Plugins/Compression/OodleNetwork/Source/Public/OodleNetworkFaultHandler.h
@@ -5,6 +5,9 @@
 // Includes
 #include "Net/Core/Connection/NetResultManager.h"
 
+#include "OodleNetworkFaultHandler.generated.h"
+
+
 // Forward declarations
 namespace UE
 {
@@ -39,7 +42,7 @@ enum class EOodleNetResult : uint8
 
 DECLARE_NETRESULT_ENUM(EOodleNetResult);
 
-OODLENETWORKHANDLERCOMPONENT_API const TCHAR* LexToString(EOodleNetResult InResult);
+OODLENETWORKHANDLERCOMPONENT_API const TCHAR* LexToString(EOodleNetResult Enum);
 
 
 /**
diff --git a/Engine/Plugins/Runtime/PacketHandlers/AESGCMHandlerComponent/Source/Private/AESGCMFaultHandler.cpp b/Engine/Plugins/Runtime/PacketHandlers/AESGCMHandlerComponent/Source/Private/AESGCMFaultHandler.cpp
--- a/Engine/Plugins/Runtime/PacketHandlers/AESGCMHandlerComponent/Source/Private/AESGCMFaultHandler.cpp
+++ b/Engine/Plugins/Runtime/PacketHandlers/AESGCMHandlerComponent/Source/Private/AESGCMFaultHandler.cpp
@@ -10,35 +10,27 @@
  * EAESGCMNetResult
  */
 
-const TCHAR* LexToString(EAESGCMNetResult InResult)
-{
-	switch (InResult)
-	{
-	case EAESGCMNetResult::Unknown:
-		return TEXT("Unknown");
-
-	case EAESGCMNetResult::Success:
-		return TEXT("Success");
-
-	case EAESGCMNetResult::AESMissingIV:
-		return TEXT("AESMissingIV");
-
-	case EAESGCMNetResult::AESMissingAuthTag:
-		return TEXT("AESMissingAuthTag");
+#ifndef CASE_ENUM_TO_TEXT_RET
+#define CASE_ENUM_TO_TEXT_RET(txt) case txt: ReturnVal = TEXT(#txt); break;
+#endif
 
-	case EAESGCMNetResult::AESMissingPayload:
-		return TEXT("AESMissingPayload");
+const TCHAR* LexToString(EAESGCMNetResult Enum)
+{
+	const TCHAR* ReturnVal = TEXT("::Invalid");
 
-	case EAESGCMNetResult::AESDecryptionFailed:
-		return TEXT("AESDecryptionFailed");
+	switch (Enum)
+	{
+		FOREACH_ENUM_EAESGCMNETRESULT(CASE_ENUM_TO_TEXT_RET)
+	}
 
-	case EAESGCMNetResult::AESZeroLastByte:
-		return TEXT("AESZeroLastByte");
+	while (*ReturnVal != ':')
+	{
+		ReturnVal++;
+	}
 
+	ReturnVal += 2;
 
-	default:
-		return TEXT("Invalid");
-	}
+	return ReturnVal;
 }
 
 
@@ -108,35 +100,3 @@ UE::Net::EHandleNetResult FAESGCMFaultHandler::HandleNetResult(UE::Net::FNetResu
 		
 	return ReturnVal;
 }
-
-#if !(UE_BUILD_SHIPPING || UE_BUILD_TEST)
-#include "Misc/AutomationTest.h"
-
-IMPLEMENT_SIMPLE_AUTOMATION_TEST(FAESGCMNetResultEnumTest, "System.Core.Networking.EAESGCMNetResult.EnumTest",
-									EAutomationTestFlags::ApplicationContextMask | EAutomationTestFlags::SmokeFilter);
-
-bool FAESGCMNetResultEnumTest::RunTest(const FString& Parameters)
-{
-	const UEnum* AESGCMResultEnum = StaticEnum<EAESGCMNetResult>();
-
-	if (TestTrue(TEXT("EAESGCMNetResult must exist"), AESGCMResultEnum != nullptr) && AESGCMResultEnum != nullptr)
-	{
-		const int64 AESGCMResultEnumLast = AESGCMResultEnum->GetMaxEnumValue() - 1;
-		bool bLexMismatch = false;
-
-		for (int64 EnumIdx=0; EnumIdx<=AESGCMResultEnumLast; EnumIdx++)
-		{
-			if (AESGCMResultEnum->GetNameStringByValue(EnumIdx) != LexToString((EAESGCMNetResult)EnumIdx))
-			{
-				bLexMismatch = true;
-				break;
-			}
-		}
-
-		TestFalse(TEXT("EAESGCMNetResult must not be missing LexToString entries"), bLexMismatch);
-	}
-
-	return true;
-}
-#endif
-
diff --git a/Engine/Plugins/Runtime/PacketHandlers/AESGCMHandlerComponent/Source/Public/AESGCMFaultHandler.h b/Engine/Plugins/Runtime/PacketHandlers/AESGCMHandlerComponent/Source/Public/AESGCMFaultHandler.h
--- a/Engine/Plugins/Runtime/PacketHandlers/AESGCMHandlerComponent/Source/Public/AESGCMFaultHandler.h
+++ b/Engine/Plugins/Runtime/PacketHandlers/AESGCMHandlerComponent/Source/Public/AESGCMFaultHandler.h
@@ -5,6 +5,9 @@
 // Includes
 #include "Net/Core/Connection/NetResultManager.h"
 
+#include "AESGCMFaultHandler.generated.h"
+
+
 // Forward declarations
 namespace UE
 {
@@ -41,7 +44,7 @@ enum class EAESGCMNetResult : uint8
 
 DECLARE_NETRESULT_ENUM(EAESGCMNetResult);
 
-AESGCMHANDLERCOMPONENT_API const TCHAR* LexToString(EAESGCMNetResult InResult);
+AESGCMHANDLERCOMPONENT_API const TCHAR* LexToString(EAESGCMNetResult Enum);
 
 
 
