Improve performance of AES GCM encryption & decryption in packet handler.
This is done by adding ability to reset PlatformCryptoEncrypto/Decryptor cipher with new IV values.

The performance improvements are from:
1) Doing initialization of encryption/decryption cipher only once, when key is set, not per each encryption/decryption operation.
2) Not calling platform random number generator for IV in each encryption of outgoing packet.
3) Some minor changes to avoid heap allocations for buffers with known size - changing TArray to TStaticArray.

Random bytes for IV in encryption are now generated only once, when key is set. Afterwards uniqueness of IV is done by incrementing IV bytes as 64-bit unsigned integer with wrap-around.

diff --git a/Engine/Plugins/Experimental/PlatformCrypto/Source/PlatformCryptoOpenSSL/Private/PlatformCryptoAesDecryptorsOpenSSL.cpp b/Engine/Plugins/Experimental/PlatformCrypto/Source/PlatformCryptoOpenSSL/Private/PlatformCryptoAesDecryptorsOpenSSL.cpp
--- a/Engine/Plugins/Experimental/PlatformCrypto/Source/PlatformCryptoOpenSSL/Private/PlatformCryptoAesDecryptorsOpenSSL.cpp
+++ b/Engine/Plugins/Experimental/PlatformCrypto/Source/PlatformCryptoOpenSSL/Private/PlatformCryptoAesDecryptorsOpenSSL.cpp
@@ -134,6 +134,33 @@ EPlatformCryptoResult FPlatformCryptoDecryptor_AES_Base_OpenSSL::Finalize(const
 	return EPlatformCryptoResult::Success;
 }
 
+EPlatformCryptoResult FPlatformCryptoDecryptor_AES_Base_OpenSSL::Reset(const TArrayView<const uint8> InitializationVector)
+{
+	if (State != EDecryptorState::Initialized && State != EDecryptorState::Finalized)
+	{
+		UE_LOG(LogPlatformCryptoOpenSSL, Warning, TEXT("FPlatformCryptoDecryptor_AES_Base_OpenSSL::Reset: Invalid state. Was %s, but should be Initialized or Finalized"), LexToString(State));
+		return EPlatformCryptoResult::Failure;
+	}
+
+	const int32 IVExpectedLength = EVP_CIPHER_CTX_iv_length(EVPContext.Get());
+	if (InitializationVector.Num() < IVExpectedLength)
+	{
+		UE_LOG(LogPlatformCryptoOpenSSL, Warning, TEXT("FPlatformCryptoDecryptor_AES_Base_OpenSSL::Reset: Invalid InitializationVector Size. InitializationVectorSize=[%d] Expected=[%d]"), InitializationVector.Num(), IVExpectedLength);
+		return EPlatformCryptoResult::Failure;
+	}
+
+	const uint8* InitializationVectorPtr = InitializationVector.Num() == 0 ? nullptr : InitializationVector.GetData();
+	const int InitResult = EVP_DecryptInit_ex(EVPContext.Get(), nullptr, nullptr, nullptr, InitializationVectorPtr);
+	if (InitResult != OPENSSL_CIPHER_SUCCESS)
+	{
+		UE_LOG(LogPlatformCryptoOpenSSL, Warning, TEXT("FPlatformCryptoDecryptor_AES_Base_OpenSSL::Reset: EVP_DecryptInit_ex failed. Result=[%d]"), InitResult);
+		return EPlatformCryptoResult::Failure;
+	}
+
+	State = EDecryptorState::Initialized;
+	return EPlatformCryptoResult::Success;
+}
+
 EPlatformCryptoResult FPlatformCryptoDecryptor_AES_Base_OpenSSL::Reset()
 {
 	const int ResetResult = EVP_CIPHER_CTX_reset(EVPContext.Get());
diff --git a/Engine/Plugins/Experimental/PlatformCrypto/Source/PlatformCryptoOpenSSL/Private/PlatformCryptoAesDecryptorsOpenSSL.h b/Engine/Plugins/Experimental/PlatformCrypto/Source/PlatformCryptoOpenSSL/Private/PlatformCryptoAesDecryptorsOpenSSL.h
--- a/Engine/Plugins/Experimental/PlatformCrypto/Source/PlatformCryptoOpenSSL/Private/PlatformCryptoAesDecryptorsOpenSSL.h
+++ b/Engine/Plugins/Experimental/PlatformCrypto/Source/PlatformCryptoOpenSSL/Private/PlatformCryptoAesDecryptorsOpenSSL.h
@@ -36,6 +36,7 @@ public:
 	virtual int32 GetCipherInitializationVectorSizeBytes() const override final;
 	virtual EPlatformCryptoResult Update(const TArrayView<const uint8> Ciphertext, const TArrayView<uint8> OutPlaintext, int32& OutPlaintextBytesWritten) override final;
 	virtual EPlatformCryptoResult Finalize(const TArrayView<uint8> OutPlaintext, int32& OutPlaintextBytesWritten) override final;
+	virtual EPlatformCryptoResult Reset(const TArrayView<const uint8> InitializationVector) override final;
 	//~ End IPlatformCryptoDecryptor Interface
 
 	/**
diff --git a/Engine/Plugins/Experimental/PlatformCrypto/Source/PlatformCryptoOpenSSL/Private/PlatformCryptoAesEncryptorsOpenSSL.cpp b/Engine/Plugins/Experimental/PlatformCrypto/Source/PlatformCryptoOpenSSL/Private/PlatformCryptoAesEncryptorsOpenSSL.cpp
--- a/Engine/Plugins/Experimental/PlatformCrypto/Source/PlatformCryptoOpenSSL/Private/PlatformCryptoAesEncryptorsOpenSSL.cpp
+++ b/Engine/Plugins/Experimental/PlatformCrypto/Source/PlatformCryptoOpenSSL/Private/PlatformCryptoAesEncryptorsOpenSSL.cpp
@@ -141,6 +141,33 @@ EPlatformCryptoResult FPlatformCryptoEncryptor_AES_Base_OpenSSL::Finalize(const
 	return EPlatformCryptoResult::Success;
 }
 
+EPlatformCryptoResult FPlatformCryptoEncryptor_AES_Base_OpenSSL::Reset(const TArrayView<const uint8> InitializationVector)
+{
+	if (State != EEncryptorState::Initialized && State != EEncryptorState::Finalized)
+	{
+		UE_LOG(LogPlatformCryptoOpenSSL, Warning, TEXT("FPlatformCryptoEncryptor_AES_Base_OpenSSL::Reset: Invalid state. Was %s, but should be Initialized or Finalized"), LexToString(State));
+		return EPlatformCryptoResult::Failure;
+	}
+
+	const int32 IVExpectedLength = EVP_CIPHER_CTX_iv_length(EVPContext.Get());
+	if (InitializationVector.Num() < IVExpectedLength)
+	{
+		UE_LOG(LogPlatformCryptoOpenSSL, Warning, TEXT("FPlatformCryptoDecryptor_AES_Base_OpenSSL::Reset: Invalid InitializationVector Size. InitializationVectorSize=[%d] Expected=[%d]"), InitializationVector.Num(), IVExpectedLength);
+		return EPlatformCryptoResult::Failure;
+	}
+
+	const uint8* InitializationVectorPtr = InitializationVector.Num() == 0 ? nullptr : InitializationVector.GetData();
+	const int InitResult = EVP_EncryptInit_ex(EVPContext.Get(), nullptr, nullptr, nullptr, InitializationVectorPtr);
+	if (InitResult != OPENSSL_CIPHER_SUCCESS)
+	{
+		UE_LOG(LogPlatformCryptoOpenSSL, Warning, TEXT("FPlatformCryptoEncryptor_AES_Base_OpenSSL::Reset: EVP_EncryptInit_ex failed. Result=[%d]"), InitResult);
+		return EPlatformCryptoResult::Failure;
+	}
+
+	State = EEncryptorState::Initialized;
+	return EPlatformCryptoResult::Success;
+}
+
 EPlatformCryptoResult FPlatformCryptoEncryptor_AES_Base_OpenSSL::Reset()
 {
 	const int ResetResult = EVP_CIPHER_CTX_reset(EVPContext.Get());
diff --git a/Engine/Plugins/Experimental/PlatformCrypto/Source/PlatformCryptoOpenSSL/Private/PlatformCryptoAesEncryptorsOpenSSL.h b/Engine/Plugins/Experimental/PlatformCrypto/Source/PlatformCryptoOpenSSL/Private/PlatformCryptoAesEncryptorsOpenSSL.h
--- a/Engine/Plugins/Experimental/PlatformCrypto/Source/PlatformCryptoOpenSSL/Private/PlatformCryptoAesEncryptorsOpenSSL.h
+++ b/Engine/Plugins/Experimental/PlatformCrypto/Source/PlatformCryptoOpenSSL/Private/PlatformCryptoAesEncryptorsOpenSSL.h
@@ -38,6 +38,7 @@ public:
 	virtual EPlatformCryptoResult GenerateAuthTag(const TArrayView<uint8> OutAuthTag, int32& OutAuthTagBytesWritten) const override final;
 	virtual EPlatformCryptoResult Update(const TArrayView<const uint8> Plaintext, const TArrayView<uint8> OutCiphertext, int32& OutCiphertextBytesWritten) override final;
 	virtual EPlatformCryptoResult Finalize(const TArrayView<uint8> OutCiphertext, int32& OutCiphertextBytesWritten) override final;
+	virtual EPlatformCryptoResult Reset(const TArrayView<const uint8> InitializationVector) override final;
 	//~ End IPlatformCryptoEncryptor Interface
 
 	/**
diff --git a/Engine/Plugins/Experimental/PlatformCrypto/Source/PlatformCryptoTypes/Public/PlatformCryptoTypes.h b/Engine/Plugins/Experimental/PlatformCrypto/Source/PlatformCryptoTypes/Public/PlatformCryptoTypes.h
--- a/Engine/Plugins/Experimental/PlatformCrypto/Source/PlatformCryptoTypes/Public/PlatformCryptoTypes.h
+++ b/Engine/Plugins/Experimental/PlatformCrypto/Source/PlatformCryptoTypes/Public/PlatformCryptoTypes.h
@@ -102,6 +102,14 @@ public:
 	 * @return Success if input was valid and the final block was encrypted successfully, or Failure otherwise
 	 */
 	virtual EPlatformCryptoResult Finalize(const TArrayView<uint8> OutCiphertext, int32& OutCiphertextBytesWritten) = 0;
+
+	/**
+	 * Resets block cipher. This allows to reuse encryption cipher with the same key for different stream with different InitializationVector.
+	 *
+	 * @param Optional IV bytes. Will be used only for block modes that requires InitializationVector
+	 * @return Success if cipher was reset successfully, or Failure otherwise
+	 */
+	virtual EPlatformCryptoResult Reset(const TArrayView<const uint8> InitializationVector) = 0;
 };
 
 /**
@@ -186,6 +194,14 @@ public:
 	 * @return Success if input was valid and the final block was decrypted successfully, or Failure otherwise
 	 */
 	virtual EPlatformCryptoResult Finalize(const TArrayView<uint8> OutPlaintext, int32& OutPlaintextBytesWritten) = 0;
+
+	/**
+	 * Resets block cipher. This allows to reuse decryption cipher with the same key for different stream with different InitializationVector.
+	 *
+	 * @param Optional IV bytes. Will be used only for block modes that requires InitializationVector
+	 * @return Success if cipher was reset successfully, or Failure otherwise
+	 */
+	virtual EPlatformCryptoResult Reset(const TArrayView<const uint8> InitializationVector) = 0;
 };
 
 /**
diff --git a/Engine/Plugins/Runtime/PacketHandlers/AESGCMHandlerComponent/Source/Private/AESGCMHandlerComponent.cpp b/Engine/Plugins/Runtime/PacketHandlers/AESGCMHandlerComponent/Source/Private/AESGCMHandlerComponent.cpp
--- a/Engine/Plugins/Runtime/PacketHandlers/AESGCMHandlerComponent/Source/Private/AESGCMHandlerComponent.cpp
+++ b/Engine/Plugins/Runtime/PacketHandlers/AESGCMHandlerComponent/Source/Private/AESGCMHandlerComponent.cpp
@@ -27,14 +27,28 @@ FAESGCMHandlerComponent::FAESGCMHandlerComponent()
 
 void FAESGCMHandlerComponent::SetEncryptionData(const FEncryptionData& EncryptionData)
 {
+	Decryptor.Reset(nullptr);
+	Encryptor.Reset(nullptr);
+
 	if (EncryptionData.Key.Num() != KeySizeInBytes)
 	{
-		UE_LOG(PacketHandlerLog, Log, TEXT("FAESGCMHandlerComponent::SetEncryptionKey. NewKey is not %d bytes long, ignoring."), KeySizeInBytes);
+		UE_LOG(PacketHandlerLog, Log, TEXT("FAESGCMHandlerComponent::SetEncryptionData. NewKey is not %d bytes long, ignoring."), KeySizeInBytes);
 		return;
 	}
 
-	Key.Reset(KeySizeInBytes);
-	Key.Append(EncryptionData.Key.GetData(), EncryptionData.Key.Num());
+	// Generate random bytes used for encryption packets
+	EPlatformCryptoResult RandResult = EncryptionContext->CreateRandomBytes(OutIV);
+	if (RandResult == EPlatformCryptoResult::Failure)
+	{
+		UE_LOG(PacketHandlerLog, Log, TEXT("FAESGCMHandlerComponent::SetEncryptionData: failed to generate IV."));
+		return;
+	}
+	
+	// Dummy IV and AuthTag values, Decrytor/Encryptor will be reset with actual values before each use
+	uint8 DummyIV[IVSizeInBytes] = { 0 };
+	uint8 DummyAuth[AuthTagSizeInBytes] = { 0 };
+	Decryptor = EncryptionContext->CreateDecryptor_AES_256_GCM(EncryptionData.Key, DummyIV, DummyAuth);
+	Encryptor = EncryptionContext->CreateEncryptor_AES_256_GCM(EncryptionData.Key, DummyIV);
 }
 
 void FAESGCMHandlerComponent::EnableEncryption()
@@ -85,17 +99,16 @@ void FAESGCMHandlerComponent::Incoming(FIncomingPacketRef PacketRef)
 		if (Packet.ReadBit() != 0)
 		{
 			// If the key hasn't been set yet, we can't decrypt, so ignore this packet. We don't set an error in this case because it may just be an out-of-order packet.
-			if (Key.Num() == 0)
+			if (!Decryptor.IsValid())
 			{
 				UE_LOG(PacketHandlerLog, Log, TEXT("FAESGCMHandlerComponent::Incoming: received encrypted packet before key was set, ignoring."));
 				Packet.SetData(nullptr, 0);
 				return;
 			}
 
+			TStaticArray<uint8, IVSizeInBytes> IV;
 			if (Packet.GetBytesLeft() >= IVSizeInBytes)
 			{
-				IV.Reset();
-				IV.AddUninitialized(IVSizeInBytes);
 				Packet.SerializeBits(IV.GetData(), IV.Num() * 8);
 			}
 			else
@@ -108,10 +121,9 @@ void FAESGCMHandlerComponent::Incoming(FIncomingPacketRef PacketRef)
 				return;
 			}
 
+			TStaticArray<uint8, AuthTagSizeInBytes> AuthTag;
 			if (Packet.GetBytesLeft() >= AuthTagSizeInBytes)
 			{
-				AuthTag.Reset();
-				AuthTag.AddUninitialized(AuthTagSizeInBytes);
 				Packet.SerializeBits(AuthTag.GetData(), AuthTag.Num() * 8);
 			}
 			else
@@ -146,7 +158,7 @@ void FAESGCMHandlerComponent::Incoming(FIncomingPacketRef PacketRef)
 			UE_LOG(PacketHandlerLog, VeryVerbose, TEXT("AESGCM packet handler received %ld bytes before decryption."), Ciphertext.Num());
 
 			EPlatformCryptoResult DecryptResult = EPlatformCryptoResult::Failure;
-			TArray<uint8> Plaintext = EncryptionContext->Decrypt_AES_256_GCM(Ciphertext, Key, IV, AuthTag, DecryptResult);
+			TArray<uint8> Plaintext = Decrypt(Ciphertext, IV, AuthTag, DecryptResult);
 
 			if (DecryptResult == EPlatformCryptoResult::Failure)
 			{
@@ -225,20 +237,16 @@ void FAESGCMHandlerComponent::Outgoing(FBitWriter& Packet, FOutPacketTraits& Tra
 				return;
 			}
 
-			TArray<uint8> OutIV;
-			OutIV.AddUninitialized(IVSizeInBytes);
-			EPlatformCryptoResult RandResult = EncryptionContext->CreateRandomBytes(OutIV);
-			if (RandResult == EPlatformCryptoResult::Failure)
-			{
-				UE_LOG(PacketHandlerLog, Log, TEXT("FAESGCMHandlerComponent::Outgoing: failed to generate IV."));
-				Packet.SetError();
-				return;
-			}
+			// Prepare new IV values for encryption.
+			// This place does not need completely new random value every time, just a unique value for each packet.
+			// Incrementing IV bytes as 64-bit integer with wrap-around is enough for this use case.
+			uint64 Counter = INTEL_ORDER64(FPlatformMemory:: ReadUnaligned<uint64>(OutIV.GetData()));
+			FPlatformMemory::WriteUnaligned(OutIV.GetData(), INTEL_ORDER64(Counter + 1));
 
-			TArray<uint8> OutAuthTag;
+			TStaticArray<uint8, AuthTagSizeInBytes> OutAuthTag;
 
 			EPlatformCryptoResult EncryptResult = EPlatformCryptoResult::Failure;
-			TArray<uint8> OutCiphertext = EncryptionContext->Encrypt_AES_256_GCM(TArrayView<uint8>(Packet.GetData(), Packet.GetNumBytes()), Key, OutIV, OutAuthTag, EncryptResult);
+			TArray<uint8> OutCiphertext = Encrypt(TArrayView<uint8>(Packet.GetData(), Packet.GetNumBytes()), OutIV, OutAuthTag, EncryptResult);
 
 			if (EncryptResult == EPlatformCryptoResult::Failure)
 			{
@@ -271,6 +279,71 @@ void FAESGCMHandlerComponent::Outgoing(FBitWriter& Packet, FOutPacketTraits& Tra
 	}
 }
 
+TArray<uint8> FAESGCMHandlerComponent::Decrypt(const TArrayView<const uint8> InCiphertext, const TArrayView<const uint8> IV, const TArrayView<const uint8> AuthTag, EPlatformCryptoResult& DecryptResult)
+{
+	DecryptResult = EPlatformCryptoResult::Failure;
+
+	TArray<uint8> OutPlaintext;
+
+	EPlatformCryptoResult Result = Decryptor->Reset(IV);
+	if (Result == EPlatformCryptoResult::Success)
+	{
+		Result = Decryptor->SetAuthTag(AuthTag);
+		if (Result == EPlatformCryptoResult::Success)
+		{
+			OutPlaintext.AddUninitialized(Decryptor->GetUpdateBufferSizeBytes(InCiphertext) + Decryptor->GetFinalizeBufferSizeBytes());
+
+			int32 UpdateBytesWritten = 0;
+			Result = Decryptor->Update(InCiphertext, OutPlaintext, UpdateBytesWritten);
+			if (Result == EPlatformCryptoResult::Success)
+			{
+				int32 FinalizeBytesWritten = 0;
+				Result = Decryptor->Finalize(TArrayView<uint8>(OutPlaintext.GetData() + UpdateBytesWritten, OutPlaintext.Num() - UpdateBytesWritten), FinalizeBytesWritten);
+				if (Result == EPlatformCryptoResult::Success)
+				{
+					OutPlaintext.SetNum(UpdateBytesWritten + FinalizeBytesWritten);
+					DecryptResult = EPlatformCryptoResult::Success;
+				}
+			}
+		}
+	}
+
+	return OutPlaintext;
+}
+
+TArray<uint8> FAESGCMHandlerComponent::Encrypt(const TArrayView<const uint8> InPlaintext, const TArrayView<const uint8> IV, TArrayView<uint8> OutAuthTag, EPlatformCryptoResult& EncryptResult)
+{
+	EncryptResult = EPlatformCryptoResult::Failure;
+
+	TArray<uint8> OutCiphertext;
+
+	EPlatformCryptoResult Result = Encryptor->Reset(IV);
+	if (Result == EPlatformCryptoResult::Success)
+	{
+		OutCiphertext.AddUninitialized(Encryptor->GetUpdateBufferSizeBytes(InPlaintext) + Encryptor->GetFinalizeBufferSizeBytes());
+
+		int32 UpdateBytesWritten = 0;
+		Result = Encryptor->Update(InPlaintext, OutCiphertext, UpdateBytesWritten);
+		if (Result == EPlatformCryptoResult::Success)
+		{
+			int32 FinalizeBytesWritten = 0;
+			Result = Encryptor->Finalize(TArrayView<uint8>(OutCiphertext.GetData() + UpdateBytesWritten, OutCiphertext.Num() - UpdateBytesWritten), FinalizeBytesWritten);
+			if (Result == EPlatformCryptoResult::Success)
+			{
+				int32 AuthTagBytesWritten = 0;
+				Result = Encryptor->GenerateAuthTag(OutAuthTag, AuthTagBytesWritten);
+				if (Result == EPlatformCryptoResult::Success)
+				{
+					OutCiphertext.SetNum(UpdateBytesWritten + FinalizeBytesWritten);
+					EncryptResult = EPlatformCryptoResult::Success;
+				}
+			}
+		}
+	}
+
+	return OutCiphertext;
+}
+
 int32 FAESGCMHandlerComponent::GetReservedPacketBits() const
 {
 	// Worst case includes the encryption enabled bit, the termination bit, padding up to the next whole byte, and a block of padding.
@@ -295,6 +368,5 @@ void FAESGCMHandlerComponent::CountBytes(FArchive& Ar) const
 	}
 	*/
 
-	Key.CountBytes(Ar);
 	Ciphertext.CountBytes(Ar);
 }
diff --git a/Engine/Plugins/Runtime/PacketHandlers/AESGCMHandlerComponent/Source/Public/AESGCMHandlerComponent.h b/Engine/Plugins/Runtime/PacketHandlers/AESGCMHandlerComponent/Source/Public/AESGCMHandlerComponent.h
--- a/Engine/Plugins/Runtime/PacketHandlers/AESGCMHandlerComponent/Source/Public/AESGCMHandlerComponent.h
+++ b/Engine/Plugins/Runtime/PacketHandlers/AESGCMHandlerComponent/Source/Public/AESGCMHandlerComponent.h
@@ -5,6 +5,7 @@
 #include "EncryptionComponent.h"
 #include "AESGCMFaultHandler.h"
 #include "IPlatformCrypto.h" // IWYU pragma: keep
+#include "Containers/StaticArray.h"
 
 struct FBitWriter;
 struct FEncryptionData;
@@ -56,17 +57,22 @@ public:
 private:
 	TUniquePtr<FEncryptionContext> EncryptionContext;
 
-	TArray<uint8> Key;
+	TUniquePtr<IPlatformCryptoDecryptor> Decryptor;
+	TUniquePtr<IPlatformCryptoEncryptor> Encryptor;
 	
 	// Avoid per packet allocations
 	TArray<uint8> Ciphertext;
-	TArray<uint8> IV;
-	TArray<uint8> AuthTag;
+
+	// IV used for encryption
+	TStaticArray<uint8, IVSizeInBytes> OutIV;
 
 	bool bEncryptionEnabled;
 
 	/** Fault handler for AESGCM-specific errors, that may trigger NetConnection Close */
 	FAESGCMFaultHandler AESGCMFaultHandler;
+
+	TArray<uint8> Decrypt(const TArrayView<const uint8> InCiphertext, const TArrayView<const uint8> IV, const TArrayView<const uint8> AuthTag, EPlatformCryptoResult& DecryptResult);
+	TArray<uint8> Encrypt(const TArrayView<const uint8> InPlaintext, const TArrayView<const uint8> IV, TArrayView<uint8> OutAuthTag, EPlatformCryptoResult& EncryptResult);
 };
 
 
