Fixing single file compile.

diff --git a/Engine/Plugins/Runtime/PacketHandlers/AESGCMHandlerComponent/Source/Public/AESGCMHandlerComponent.h b/Engine/Plugins/Runtime/PacketHandlers/AESGCMHandlerComponent/Source/Public/AESGCMHandlerComponent.h
--- a/Engine/Plugins/Runtime/PacketHandlers/AESGCMHandlerComponent/Source/Public/AESGCMHandlerComponent.h
+++ b/Engine/Plugins/Runtime/PacketHandlers/AESGCMHandlerComponent/Source/Public/AESGCMHandlerComponent.h
@@ -14,7 +14,7 @@ struct FOutPacketTraits;
 /*
 * AES256 GCM block encryption component.
 */
-class AESGCMHANDLERCOMPONENT_API FAESGCMHandlerComponent : public FEncryptionComponent
+class FAESGCMHandlerComponent : public FEncryptionComponent
 {
 public:
 	/**
@@ -22,7 +22,7 @@ public:
 	 * You must set the key before enabling encryption, or before receiving encrypted
 	 * packets, or those operations will fail.
 	 */
-	FAESGCMHandlerComponent();
+	AESGCMHANDLERCOMPONENT_API FAESGCMHandlerComponent();
 
 	// This handler uses AES256, which has 32-byte keys.
 	static const int32 KeySizeInBytes = 32;
@@ -34,25 +34,25 @@ public:
 	static const int32 AuthTagSizeInBytes = 16;
 
 	// Replace the key used for encryption with NewKey if NewKey is exactly KeySizeInBytes long.
-	virtual void SetEncryptionData(const FEncryptionData& EncryptionData) override;
+	AESGCMHANDLERCOMPONENT_API virtual void SetEncryptionData(const FEncryptionData& EncryptionData) override;
 
 	// After calling this, future outgoing packets will be encrypted (until a call to DisableEncryption).
-	virtual void EnableEncryption() override;
+	AESGCMHANDLERCOMPONENT_API virtual void EnableEncryption() override;
 
 	// After calling this, future outgoing packets will not be encrypted (until a call to DisableEncryption).
-	virtual void DisableEncryption() override;
+	AESGCMHANDLERCOMPONENT_API virtual void DisableEncryption() override;
 
 	// Returns true if encryption is currently enabled.
-	virtual bool IsEncryptionEnabled() const override;
+	AESGCMHANDLERCOMPONENT_API virtual bool IsEncryptionEnabled() const override;
 
 	// HandlerComponent interface
-	virtual void Initialize() override;
-	virtual void InitFaultRecovery(UE::Net::FNetConnectionFaultRecoveryBase* InFaultRecovery) override;
-	virtual bool IsValid() const override;
-	virtual void Incoming(FIncomingPacketRef PacketRef) override;
-	virtual void Outgoing(FBitWriter& Packet, FOutPacketTraits& Traits) override;
-	virtual int32 GetReservedPacketBits() const override;
-	virtual void CountBytes(FArchive& Ar) const override;
+	AESGCMHANDLERCOMPONENT_API virtual void Initialize() override;
+	AESGCMHANDLERCOMPONENT_API virtual void InitFaultRecovery(UE::Net::FNetConnectionFaultRecoveryBase* InFaultRecovery) override;
+	AESGCMHANDLERCOMPONENT_API virtual bool IsValid() const override;
+	AESGCMHANDLERCOMPONENT_API virtual void Incoming(FIncomingPacketRef PacketRef) override;
+	AESGCMHANDLERCOMPONENT_API virtual void Outgoing(FBitWriter& Packet, FOutPacketTraits& Traits) override;
+	AESGCMHANDLERCOMPONENT_API virtual int32 GetReservedPacketBits() const override;
+	AESGCMHANDLERCOMPONENT_API virtual void CountBytes(FArchive& Ar) const override;
 
 private:
 	TUniquePtr<FEncryptionContext> EncryptionContext;
