Simplify network packet handler code

Now code uses less memory, does fewer allocations and packs payload data more efficiently.
Oodle & AES GCM handlers now store payload on exact byte boundaries, no bit granularity. On average this saves 2 bytes per packet.
Oodle stores length encoded in 1 or 2 byte prefix. One byte for length less than 252, otherwise two bytes. Uncompressed packet has one byte with value 0 as prefix.
AES GCM ensures IV vector has non-zero first byte. Unencrypted packets just store a zero first byte in front. Encrypted ones will have only IV+AuthTag overhead, exactly 28 bytes.
AES GCM handler had overestimated reserved bit count - it had unnecessary full AES block there. Removing that now allows to pack extra 16 bytes of payload into each packet.

diff --git a/Engine/Plugins/Compression/OodleNetwork/Source/Private/OodleNetworkHandlerComponent.cpp b/Engine/Plugins/Compression/OodleNetwork/Source/Private/OodleNetworkHandlerComponent.cpp
--- a/Engine/Plugins/Compression/OodleNetwork/Source/Private/OodleNetworkHandlerComponent.cpp
+++ b/Engine/Plugins/Compression/OodleNetwork/Source/Private/OodleNetworkHandlerComponent.cpp
@@ -185,52 +185,6 @@ static OO_BOOL STDCALL UEOodleDisplayAssert(const char* File, const int Line, co
 	return false;
 }
 
-
-/**
- * Serialization functions which allow PacketSize == MAX_OODLE_PACKET_SIZE, by assuming PacketSize is never 0
- */
-
-FORCEINLINE void SerializeOodlePacketSize(FBitWriter& Writer, uint32 PacketSize)
-{
-	if (PacketSize > 0)
-	{
-		PacketSize--;
-
-		// @todo #JohnB: Restore when serialize changes are stable
-		//Writer.NetSerializeInt<MAX_OODLE_PACKET_BYTES>(PacketSize);
-
-		// @todo #JohnB: Remove when restoring the above
-		Writer.SerializeInt(PacketSize, MAX_OODLE_PACKET_BYTES);
-	}
-	else
-	{
-		Writer.SetError();
-
-		UE_LOG(OodleNetworkHandlerComponentLog, Error, TEXT("Oodle attempted to process zero-size packet."));
-	}
-}
-
-FORCEINLINE void SerializeOodlePacketSize(FBitReader& Reader, uint32& OutPacketSize)
-{
-	// @todo Oodle FIX ME : these serializes are just raw bits up to MAX_PACKET_SIZE = 1024 (so 10 bits)
-
-	// @todo #JohnB: Restore when serialize changes are stable
-	//Reader.NetSerializeInt<MAX_OODLE_PACKET_BYTES>(OutPacketSize);
-
-	// @todo #JohnB: Remove when restoring the above
-	Reader.SerializeInt(OutPacketSize, MAX_OODLE_PACKET_BYTES);
-
-	if (!Reader.IsError())
-	{
-		OutPacketSize++;
-	}
-	else
-	{
-		OutPacketSize = 0;	
-	}
-}
-
-
 #if STATS
 /** The global net stats tracker for Oodle */
 static FOodleNetStats GOodleNetStats;
@@ -372,7 +326,6 @@ OodleNetworkHandlerComponent::OodleNetworkHandlerComponent()
 	, bUseDictionaryIfPresent(false)
 	, bCaptureMode(false)
 #endif
-	, OodleReservedPacketBits(0)
 	, NetAnalyticsData()
 	, bOodleNetworkAnalytics(false)
 	, ServerDictionary()
@@ -550,32 +503,6 @@ void OodleNetworkHandlerComponent::Initialize()
 		}
 	}
 
-	if (bEnableOodle)
-	{
-		// Add a bit for bCompressedPacket
-		OodleReservedPacketBits = 1;
-
-		// Oodle writes the decompressed packet size, as its addition to the protocol - it writes using SerializeInt however,
-		// so determine the worst case number of packed bits that will be written, based on the MAX_OODLE_PACKET_SIZE packet limit.
-		FBitWriter MeasureAr(0, true);
-		uint32 MaxOodlePacket = MAX_OODLE_PACKET_BYTES;
-
-		SerializeOodlePacketSize(MeasureAr, MaxOodlePacket);
-
-		if (!MeasureAr.IsError())
-		{
-			OodleReservedPacketBits += static_cast<uint32>(MeasureAr.GetNumBits());
-
-#if !UE_BUILD_SHIPPING
-			SET_DWORD_STAT(STAT_PacketReservedOodle, OodleReservedPacketBits);
-#endif
-		}
-		else
-		{
-			LowLevelFatalError(TEXT("Failed to determine OodleNetworkHandlerComponent reserved packet bits."));
-		}
-	}
-
 	Initialized();
 }
 
@@ -1026,10 +953,11 @@ void OodleNetworkHandlerComponent::Incoming(FIncomingPacketRef PacketRef)
 
 	if (bEnableOodle)
 	{
-		uint8 bCompressedPacket = Packet.ReadBit();
+		const uint8* PacketData = Packet.GetData();
+		uint32 PacketBytes = Packet.GetNumBytes();
 
-		// If the packet is not compressed, no further processing is necessary
-		if (bCompressedPacket)
+		// If first byte of packet data is non-zero, then packet is compressed
+		if (PacketBytes > 0 && PacketData[0] != 0)
 		{
 			const bool bIsServer = (Handler->Mode == UE::Handler::Mode::Server);
 
@@ -1044,39 +972,35 @@ void OodleNetworkHandlerComponent::Incoming(FIncomingPacketRef PacketRef)
 			if (CurDict != nullptr)
 			{
 				uint32 BeforeDecompressedLength = static_cast<uint32>(Packet.GetBytesLeft());
-				uint32 DecompressedLength;
-
-				SerializeOodlePacketSize(Packet, DecompressedLength);
 
-				if ( DecompressedLength == 0 )
+				// Decode uncompressed length of payload
+				uint32 ExtraBytes = 1;
+				uint32 DecompressedLength = PacketData[0];
+				if (DecompressedLength >= 252)
 				{
+					if (PacketBytes < 2)
+					{
 #if !UE_BUILD_SHIPPING
-					UE_LOG(OodleNetworkHandlerComponentLog, Error,
-						TEXT("Couldn't SerializeOodlePacketSize"), 0);
+						UE_LOG(OodleNetworkHandlerComponentLog, Error, TEXT("Couldn't decode uncompressed length"), 0);
 #endif
-					return;
+						return;
+					}
+
+					ExtraBytes = 2;
+					DecompressedLength += PacketData[1] << 2;
 				}
 
-#if !UE_BUILD_SHIPPING
 				// Never allow DecompressedLength values bigger than this, due to performance/security considerations
-				check(MAX_OODLE_PACKET_BYTES <= 16384);
-#endif
+				static_assert(MAX_OODLE_PACKET_BYTES <= 16384, "Oodle packet max size is too big");
 
 				if (DecompressedLength < MAX_OODLE_PACKET_BYTES)
 				{
-					// @todo Oodle FIX ME static buffers = not thread safe
-					//	can we gaurantee this is done single threaded?
-					static uint8 CompressedData[MAX_OODLE_BUFFER];
-					static uint8 DecompressedData[MAX_OODLE_BUFFER];
-
-					const uint32 CompressedLength = static_cast<uint32>(Packet.GetBytesLeft());
+					uint32 CompressedLength = PacketBytes - ExtraBytes;
+					const uint8* CompressedData = PacketData + ExtraBytes;
 
-					// @todo Oodle FIX ME : byte buffer serialized out of bit buffer
-					Packet.Serialize(CompressedData, CompressedLength);
+					uint8 DecompressedData[MAX_OODLE_BUFFER];
 
-					bool bSuccess = !Packet.IsError();
-
-					if (bSuccess)
+					bool bSuccess;
 					{
 						UE_OODLE_LIGHTWEIGHT_TIME_GUARD_DECLARE(Incoming, UE::Oodle::GOodleTimeguardThresholdMS);
 
@@ -1112,15 +1036,6 @@ void OodleNetworkHandlerComponent::Incoming(FIncomingPacketRef PacketRef)
 							AddToChainResultPtr(Traits.ExtendedError, EOodleNetResult::OodleDecodeFailed);
 						}
 					}
-					else
-					{
-#if !UE_BUILD_SHIPPING
-						UE_LOG(OodleNetworkHandlerComponentLog, Error, TEXT("Error serializing received packet data"));
-#endif
-
-						Packet.SetError();
-						AddToChainResultPtr(Traits.ExtendedError, EOodleNetResult::OodleSerializePayloadFail);
-					}
 
 					if (bSuccess)
 					{
@@ -1181,6 +1096,9 @@ void OodleNetworkHandlerComponent::Incoming(FIncomingPacketRef PacketRef)
 		}
 		else
 		{
+			// Skip first zero byte, rest of bytes contains uncompressed data
+			Packet.Skip(8);
+
 			if (bOodleNetworkAnalytics && NetAnalyticsData.IsValid())
 			{
 				FOodleNetworkAnalyticsVars* AnalyticsVars = NetAnalyticsData->GetLocalData();
@@ -1235,14 +1153,6 @@ void OodleNetworkHandlerComponent::Outgoing(FBitWriter& Packet, FOutPacketTraits
 		}
 #endif
 
-		// @todo #JohnB: You should be able to share the same two buffers between Incoming and Outgoing, as an optimization
-		// @todo Oodle CB : static buffer, ugly, not thread safe
-		//	 presumably this function must be called single threaded as there's lots of non-thread-safe stuff
-		//	clarify that and clean this up
-		//	don't use static buffers, put them in the OodleNetworkHandlerComponent
-		static uint8 UncompressedData[MAX_OODLE_BUFFER];
-		static uint8 CompressedData[MAX_OODLE_BUFFER];
-
 		FOodleNetworkAnalyticsVars* AnalyticsVars = (bOodleNetworkAnalytics && NetAnalyticsData.IsValid()) ? NetAnalyticsData->GetLocalData() : nullptr;
 		const bool bIsServer = (Handler->Mode == UE::Handler::Mode::Server);
 		FOodleNetworkDictionary* CurDict = (bIsServer ? ServerDictionary.Get() : ClientDictionary.Get());
@@ -1275,7 +1185,7 @@ void OodleNetworkHandlerComponent::Outgoing(FBitWriter& Packet, FOutPacketTraits
 #if !UE_BUILD_SHIPPING
 			check(MaxOutgoingBits <= (MAX_OODLE_PACKET_BYTES * 8));
 #endif
-			uint32 MaxAdjustedLengthBits = MaxOutgoingBits - OodleReservedPacketBits;
+			uint32 MaxAdjustedLengthBits = MaxOutgoingBits - GetReservedPacketBits();
 			uint32 UncompressedBits = static_cast<uint32>(Packet.GetNumBits());
 
 			bool bWithinBitBounds = UncompressedBits > 0 && ensure(UncompressedBits <= MaxAdjustedLengthBits) &&
@@ -1283,11 +1193,10 @@ void OodleNetworkHandlerComponent::Outgoing(FBitWriter& Packet, FOutPacketTraits
 
 			if (bWithinBitBounds)
 			{
-				SSIZE_T CompressedLengthSINT = 0;
+				uint8* UncompressedData = Packet.GetData();
+				uint8 CompressedData[MAX_OODLE_BUFFER];
 
-				// @todo #JohnB: Redundant memcpy?
-				FMemory::Memcpy(UncompressedData, Packet.GetData(), UncompressedBytes);
-					
+				OO_SINTa CompressedLengthSINT;
 				{
 #if !UE_BUILD_SHIPPING
 					SCOPE_CYCLE_COUNTER(STAT_Oodle_OutCompressTime);
@@ -1302,28 +1211,51 @@ void OodleNetworkHandlerComponent::Outgoing(FBitWriter& Packet, FOutPacketTraits
 
 				if (CompressedBytes <= UncompressedBytes)
 				{
+					// Max uncompressed length is 1024 = 10 bits
+					// If length is less then 252, that requires just one byte to encode
+					// Otherwise, (length-252) value is encoded in 2 low bits of first byte, and high 8 bits in second byte
+					// Thus max value this can encode is 255 + (255 << 2) = 1275
+					static_assert(MAX_OODLE_PACKET_BYTES <= 255 + (255 << 2), "Oodle packet max size is too big");
+
+					int32 ExtraBytes = (UncompressedBytes < 252) ? 1 : 2;
+
 					// It's possible for the packet to be within bit bounds, but to overstep bounds when rounded-up to nearest byte,
 					// after processing by Oodle.
 					// If this happens, the packet will fail to fit if Oodle failed to compress enough - so will be sent uncompressed.
-					bWithinBitBounds = (CompressedBytes * 8) <= MaxAdjustedLengthBits;
-
+					bWithinBitBounds = (CompressedBytes * 8) <= (MaxOutgoingBits - ExtraBytes * 8);
 
 					// Don't write the compressed data, if it's not within bit bounds, or compression failed to provide savings
 					uint8 bCompressedPacket = bWithinBitBounds && (CompressedBytes < UncompressedBytes);
-
-					Packet.Reset();
-					Packet.WriteBit(bCompressedPacket);
-
 					Traits.bIsCompressed = !!bCompressedPacket;
 
 					if (bCompressedPacket)
 					{
-						// @todo #JohnB: Compress directly into a (deliberately oversized) FBitWriter buffer, which you can shrink after
-						//					(need to be careful of bit order)
+						// Make sure there is enough space allocated for outgoing packet memory
+						int64 NewPacketByteCount = ExtraBytes + CompressedBytes;
+						if (NewPacketByteCount * 8 > Packet.GetMaxBits())
+						{
+							// Allocate max MAX_PACKET_SIZE bytes, just like PacketHandler does, so packet memory can be reused easier
+							check(NewPacketByteCount <= MAX_PACKET_SIZE);
+							Packet = FBitWriter(MAX_PACKET_SIZE * 8);
+						}
 
-						SerializeOodlePacketSize(Packet, UncompressedBytes);
+						uint8* NewPacketData = Packet.GetData();
+						if (ExtraBytes == 1)
+						{
+							// This value will always be non-zero
+							*NewPacketData++ = static_cast<uint8>(UncompressedBytes);
+						}
+						else
+						{
+							*NewPacketData++ = static_cast<uint8>(UncompressedBytes | 252);
+							*NewPacketData++ = static_cast<uint8>((UncompressedBytes - 252) >> 2);
+						}
+
+						// Copy compressed bytes to new packet
+						FPlatformMemory::Memcpy(NewPacketData, CompressedData, CompressedBytes);
 
-						Packet.Serialize(CompressedData, CompressedBytes);
+						// Set valid bits in new packet
+						Packet.SetNumBits((ExtraBytes + CompressedBytes) * 8);
 
 #if STATS
 						GOodleNetStats.OutgoingStats(CompressedBytes, UncompressedBytes);
@@ -1344,9 +1276,23 @@ void OodleNetworkHandlerComponent::Outgoing(FBitWriter& Packet, FOutPacketTraits
 							AnalyticsVars->OutNotCompressedFailedLengthTotal += UncompressedBytes;
 						}
 
-						// @todo #JohnB: Try to eliminate this appBitscpy, by reserving the bCompressedPacket bit and other data,
-						//					at the start of the bit writer
-						Packet.SerializeBits(UncompressedData, UncompressedBits);
+						// Make sure packet has space available for extra 8 bits
+						if (Packet.AllowAppend(8))
+						{
+							// Reserve one byte in packet data in the beginning
+							uint8* PacketData = Packet.GetData();
+							FPlatformMemory::Memmove(PacketData + 1, PacketData, Packet.GetNumBytes());
+
+							// First byte 0 means that packet contains uncompressed bytes
+							PacketData[0] = 0;
+
+							// Include first 8 bits in new packet data
+							Packet.SetNumBits(8 + Packet.GetNumBits());
+						}
+						else
+						{
+							Packet.SetOverflowed(Packet.GetNumBits());
+						}
 
 						if (!bWithinBitBounds)
 						{
@@ -1422,14 +1368,23 @@ void OodleNetworkHandlerComponent::Outgoing(FBitWriter& Packet, FOutPacketTraits
 				AnalyticsVars->OutNotCompressedTooSmallNum += (uint8)bSkipCompressionTooSmall;
 			}
 
-			uint8 bCompresedPacket = false;
-			uint32 UncompressedBits = static_cast<uint32>(Packet.GetNumBits());
+			// Make sure packet has space available for extra 8 bits
+			if (Packet.AllowAppend(8))
+			{
+				// Reserve one byte in packet data in the beginning
+				uint8* PacketData = Packet.GetData();
+				FPlatformMemory::Memmove(PacketData + 1, PacketData, Packet.GetNumBytes());
 
-			FMemory::Memcpy(UncompressedData, Packet.GetData(), Packet.GetNumBytes());
+				// First byte 0 means that packet contains uncompressed bytes
+				PacketData[0] = 0;
 
-			Packet.Reset();
-			Packet.WriteBit(bCompresedPacket);
-			Packet.SerializeBits(UncompressedData, UncompressedBits);
+				// Include first 8 bits in new packet data
+				Packet.SetNumBits(8 + Packet.GetNumBits());
+			}
+			else
+			{
+				Packet.SetOverflowed(Packet.GetNumBits());
+			}
 		}
 		else // if (CurDict == nullptr)
 		{
@@ -1470,7 +1425,9 @@ bool OodleNetworkHandlerComponent::IsCompressionActive() const
 
 int32 OodleNetworkHandlerComponent::GetReservedPacketBits() const
 {
-	return bEnableOodle ? OodleReservedPacketBits : 0;
+	// Worst case is 16 bits - up to two bytes that encode uncompressed packet size
+	// For uncompressed packets it will just 1 byte, which is less
+	return 16;
 }
 
 void OodleNetworkHandlerComponent::NotifyAnalyticsProvider()
diff --git a/Engine/Plugins/Compression/OodleNetwork/Source/Public/OodleNetworkHandlerComponent.h b/Engine/Plugins/Compression/OodleNetwork/Source/Public/OodleNetworkHandlerComponent.h
--- a/Engine/Plugins/Compression/OodleNetwork/Source/Public/OodleNetworkHandlerComponent.h
+++ b/Engine/Plugins/Compression/OodleNetwork/Source/Public/OodleNetworkHandlerComponent.h
@@ -385,9 +385,6 @@ protected:
 	bool bCaptureMode;
 #endif
 
-	/** Cached reserved packet bits for Oodle */
-	uint32 OodleReservedPacketBits;
-
 	/** The net analytics aggregator data, which will take the above locally tracked variables, once they are complete */
 	TNetAnalyticsDataPtr<FOodleNetAnalyticsData> NetAnalyticsData;
 
diff --git a/Engine/Plugins/Runtime/PacketHandlers/AESGCMHandlerComponent/Source/Private/AESGCMHandlerComponent.cpp b/Engine/Plugins/Runtime/PacketHandlers/AESGCMHandlerComponent/Source/Private/AESGCMHandlerComponent.cpp
--- a/Engine/Plugins/Runtime/PacketHandlers/AESGCMHandlerComponent/Source/Private/AESGCMHandlerComponent.cpp
+++ b/Engine/Plugins/Runtime/PacketHandlers/AESGCMHandlerComponent/Source/Private/AESGCMHandlerComponent.cpp
@@ -36,13 +36,17 @@ void FAESGCMHandlerComponent::SetEncryptionData(const FEncryptionData& Encryptio
 		return;
 	}
 
-	// Generate random bytes used for encryption packets
-	EPlatformCryptoResult RandResult = EncryptionContext->CreateRandomBytes(OutIV);
-	if (RandResult == EPlatformCryptoResult::Failure)
+	// Generate random bytes used for encryption packets, make sure first IV byte is non-zero value
+	do
 	{
-		UE_LOG(PacketHandlerLog, Log, TEXT("FAESGCMHandlerComponent::SetEncryptionData: failed to generate IV."));
-		return;
+		EPlatformCryptoResult RandResult = EncryptionContext->CreateRandomBytes(OutIV);
+		if (RandResult == EPlatformCryptoResult::Failure)
+		{
+			UE_LOG(PacketHandlerLog, Log, TEXT("FAESGCMHandlerComponent::SetEncryptionData: failed to generate IV."));
+			return;
+		}
 	}
+	while (OutIV[0] == 0);
 	
 	// Dummy IV and AuthTag values, Decrytor/Encryptor will be reset with actual values before each use
 	uint8 DummyIV[IVSizeInBytes] = { 0 };
@@ -83,6 +87,20 @@ bool FAESGCMHandlerComponent::IsValid() const
 	return true;
 }
 
+
+// Incoming & Outgoing pakcets will always operate on byte level. No bit-aligned stuff is supported here for efficiency reasons.
+
+// Encrypted packed layout: 
+// [iv:12] [auth:16] [ciphertext:N]
+// where low byte of first IV byte is non-zero value
+
+// Unencrypted packed layout:
+// [0] [plaintext:N]
+// where first byte is 0 to signal unencrypted packet
+
+// Any further packet handler is expected to encode their bit-length explicitly if they operate on non-byte aligned packet sizes
+// In typical use case there will be OodleNetworkHandler that is byte aligned and it encodes uncompressed length
+
 void FAESGCMHandlerComponent::Incoming(FIncomingPacketRef PacketRef)
 {
 	using namespace UE::Net;
@@ -92,11 +110,15 @@ void FAESGCMHandlerComponent::Incoming(FIncomingPacketRef PacketRef)
 	FBitReader& Packet = PacketRef.Packet;
 	FInPacketTraits& Traits = PacketRef.Traits;
 
+	int64 PacketBytes = Packet.GetNumBytes();
+
 	// Handle this packet
-	if (IsValid() && Packet.GetNumBytes() > 0)
+	if (IsValid() && PacketBytes > 0)
 	{
-		// Check first bit to see whether payload is encrypted
-		if (Packet.ReadBit() != 0)
+		uint8* PacketData = Packet.GetData();
+
+		// If first byte is nonzero, then payload is encrypted
+		if (PacketData[0] != 0)
 		{
 			// If the key hasn't been set yet, we can't decrypt, so ignore this packet. We don't set an error in this case because it may just be an out-of-order packet.
 			if (!Decryptor.IsValid())
@@ -106,10 +128,12 @@ void FAESGCMHandlerComponent::Incoming(FIncomingPacketRef PacketRef)
 				return;
 			}
 
-			TStaticArray<uint8, IVSizeInBytes> IV;
-			if (Packet.GetBytesLeft() >= IVSizeInBytes)
+			// First 12 bytes is IV
+			TArrayView<uint8> IV { PacketData, IVSizeInBytes };
+			if (PacketBytes >= IV.Num())
 			{
-				Packet.SerializeBits(IV.GetData(), IV.Num() * 8);
+				PacketData += IV.Num();
+				PacketBytes -= IV.Num();
 			}
 			else
 			{
@@ -121,10 +145,12 @@ void FAESGCMHandlerComponent::Incoming(FIncomingPacketRef PacketRef)
 				return;
 			}
 
-			TStaticArray<uint8, AuthTagSizeInBytes> AuthTag;
-			if (Packet.GetBytesLeft() >= AuthTagSizeInBytes)
+			// Then there are 16 bytes of AuthTag
+			TArrayView<uint8> AuthTag { PacketData, AuthTagSizeInBytes };
+			if (PacketBytes >= AuthTag.Num())
 			{
-				Packet.SerializeBits(AuthTag.GetData(), AuthTag.Num() * 8);
+				PacketData += AuthTag.Num();
+				PacketBytes -= AuthTag.Num();
 			}
 			else
 			{
@@ -136,16 +162,7 @@ void FAESGCMHandlerComponent::Incoming(FIncomingPacketRef PacketRef)
 				return;
 			}
 
-			// Copy remaining bits to a TArray so that they are byte-aligned.
-			if (Packet.GetBytesLeft() > 0)
-			{
-				Ciphertext.Reset();
-				Ciphertext.AddUninitialized(Packet.GetBytesLeft());
-				Ciphertext[Ciphertext.Num() - 1] = 0;
-
-				Packet.SerializeBits(Ciphertext.GetData(), Packet.GetBitsLeft());
-			}
-			else
+			if (PacketBytes == 0)
 			{
 				UE_LOG(PacketHandlerLog, Log, TEXT("FAESGCMHandlerComponent::Incoming: missing ciphertext"));
 
@@ -155,11 +172,15 @@ void FAESGCMHandlerComponent::Incoming(FIncomingPacketRef PacketRef)
 				return;
 			}
 
-			UE_LOG(PacketHandlerLog, VeryVerbose, TEXT("AESGCM packet handler received %ld bytes before decryption."), Ciphertext.Num());
+			// Rest of bytes is ciphertext
+			TArrayView<uint8> CipherText { PacketData, static_cast<int32>(PacketBytes) };
+
+			UE_LOG(PacketHandlerLog, VeryVerbose, TEXT("AESGCM packet handler received %d bytes before decryption."), CipherText.Num());
 
-			EPlatformCryptoResult DecryptResult = EPlatformCryptoResult::Failure;
-			TArray<uint8> Plaintext = Decrypt(Ciphertext, IV, AuthTag, DecryptResult);
+			uint8 PlainText[MAX_PACKET_SIZE];
 
+			// Decrypt payload and verify AuthTag
+			EPlatformCryptoResult DecryptResult = Decrypt(PlainText, CipherText, IV, AuthTag);
 			if (DecryptResult == EPlatformCryptoResult::Failure)
 			{
 				UE_LOG(PacketHandlerLog, Log, TEXT("FAESGCMHandlerComponent::Incoming: failed to decrypt packet."));
@@ -170,37 +191,12 @@ void FAESGCMHandlerComponent::Incoming(FIncomingPacketRef PacketRef)
 				return;
 			}
 
-			if (Plaintext.Num() == 0)
-			{
-				Packet.SetData(nullptr, 0);
-				return;
-			}
-
-			// Look for the termination bit that was written in Outgoing() to determine the exact bit size.
-			uint8 LastByte = Plaintext.Last();
-
-			if (LastByte != 0)
-			{
-				int32 BitSize = (Plaintext.Num() * 8) - 1;
-
-				// Bit streaming, starts at the Least Significant Bit, and ends at the MSB.
-				while (!(LastByte & 0x80))
-				{
-					LastByte *= 2;
-					BitSize--;
-				}
-
-				UE_LOG(PacketHandlerLog, VeryVerbose, TEXT("  Have %d bits after decryption."), BitSize);
-
-				Packet.SetData(MoveTemp(Plaintext), BitSize);
-			}
-			else
-			{
-				UE_LOG(PacketHandlerLog, Log, TEXT("FAESGCMHandlerComponent::Incoming: malformed packet, last byte was 0."));
-
-				Packet.SetError();
-				AddToChainResultPtr(Traits.ExtendedError, EAESGCMNetResult::AESZeroLastByte);
-			}
+			Packet.SetData(PlainText, CipherText.Num() * 8);
+		}
+		else
+		{
+			// Skip first zero byte, rest of bytes contains unencrypted data
+			Packet.Skip(8);
 		}
 	}
 }
@@ -212,143 +208,178 @@ void FAESGCMHandlerComponent::Outgoing(FBitWriter& Packet, FOutPacketTraits& Tra
 	// Handle this packet
 	if (IsValid() && Packet.GetNumBytes() > 0)
 	{
-		// Allow for encryption enabled bit and termination bit. Allow resizing to account for encryption padding.
-		FBitWriter NewPacket(Packet.GetNumBits() + 2, true);
-		NewPacket.WriteBit(bEncryptionEnabled ? 1 : 0);
-
-		if (NewPacket.IsError())
-		{
-			UE_LOG(PacketHandlerLog, Log, TEXT("FAESGCMHandlerComponent::Outgoing: failed to write encryption bit."));
-			Packet.SetError();
-			return;
-		}
-
 		if (bEncryptionEnabled)
 		{
 			UE_LOG(PacketHandlerLog, VeryVerbose, TEXT("AESGCM packet handler sending %ld bits before encryption."), Packet.GetNumBits());
 
-			// Write a termination bit so that the receiving side can calculate the exact number of bits sent.
-			// Same technique used in UNetConnection.
-			Packet.WriteBit(1);
+			TStaticArray<uint8, AuthTagSizeInBytes> AuthTag;
+			TArrayView<uint8> Payload { Packet.GetData(), static_cast<int32>(Packet.GetNumBytes()) };
 
-			if (Packet.IsError())
+			// Prepare new IV for encryption
 			{
-				UE_LOG(PacketHandlerLog, Log, TEXT("FAESGCMHandlerComponent::Outgoing: failed to write termination bit."));
-				return;
-			}
-
-			// Prepare new IV values for encryption.
-			// This place does not need completely new random value every time, just a unique value for each packet.
-			// Incrementing IV bytes as 64-bit integer with wrap-around is enough for this use case.
-			uint64 Counter = INTEL_ORDER64(FPlatformMemory:: ReadUnaligned<uint64>(OutIV.GetData()));
-			FPlatformMemory::WriteUnaligned(OutIV.GetData(), INTEL_ORDER64(Counter + 1));
+				// Use 64-bit counter in bytes [1,9], leave byte [0] unchanged as that is non-zero indication that packet is encrypted
+				uint8* CounterLocation = OutIV.GetData() + 1;
 
-			TStaticArray<uint8, AuthTagSizeInBytes> OutAuthTag;
+				// This place does not need completely new random value every time, just a unique value for each packet.
+				// Incrementing IV bytes as 64-bit integer with wrap-around is enough for this use case.
+				uint64 Counter = INTEL_ORDER64(FPlatformMemory::ReadUnaligned<uint64>(CounterLocation));
+				FPlatformMemory::WriteUnaligned(CounterLocation, INTEL_ORDER64(Counter + 1));
+			}
 
-			EPlatformCryptoResult EncryptResult = EPlatformCryptoResult::Failure;
-			TArray<uint8> OutCiphertext = Encrypt(TArrayView<uint8>(Packet.GetData(), Packet.GetNumBytes()), OutIV, OutAuthTag, EncryptResult);
+			TArrayView<uint8> PlainText = { Packet.GetData(), static_cast<int32>(Packet.GetNumBytes()) };
+			uint8 CipherText[MAX_PACKET_SIZE];
 
+			// Encrypt payload & write AuthTag
+			EPlatformCryptoResult EncryptResult = Encrypt(CipherText, PlainText, OutIV, AuthTag);
 			if (EncryptResult == EPlatformCryptoResult::Failure)
 			{
 				UE_LOG(PacketHandlerLog, Log, TEXT("FAESGCMHandlerComponent::Outgoing: failed to encrypt packet."));
 				Packet.SetError();
 				return;
 			}
-			else
+			 
+			// Make sure there is enough space allocated for outgoing packet memory
+			int64 NewPacketByteCount = IVSizeInBytes + AuthTagSizeInBytes + PlainText.Num();
+			if (NewPacketByteCount * 8 > Packet.GetMaxBits())
 			{
-				NewPacket.Serialize(OutIV.GetData(), OutIV.Num());
-				NewPacket.Serialize(OutAuthTag.GetData(), OutAuthTag.Num());
-				NewPacket.Serialize(OutCiphertext.GetData(), OutCiphertext.Num());
-
-				if (NewPacket.IsError())
-				{
-					UE_LOG(PacketHandlerLog, Log, TEXT("FAESGCMHandlerComponent::Outgoing: failed to write ciphertext to packet."));
-					Packet.SetError();
-					return;
-				}
-
-				UE_LOG(PacketHandlerLog, VeryVerbose, TEXT("  AESGCM packet handler sending %" INT64_FMT " bytes after encryption."), NewPacket.GetNumBytes());
+				// Allocate max MAX_PACKET_SIZE bytes, just like PacketHandler does, so packet memory can be reused later
+				check(NewPacketByteCount <= MAX_PACKET_SIZE);
+				Packet = FBitWriter(MAX_PACKET_SIZE * 8);
 			}
+
+			uint8* NewPacketData = Packet.GetData();
+
+			// Copy the IV, AuthTag and encrypted payload to new packet
+			FPlatformMemory::Memcpy(NewPacketData, OutIV.GetData(), OutIV.Num());
+			NewPacketData += OutIV.Num();
+
+			FPlatformMemory::Memcpy(NewPacketData, AuthTag.GetData(), AuthTag.Num());
+			NewPacketData += AuthTag.Num();
+			
+			FPlatformMemory::Memcpy(NewPacketData, CipherText, PlainText.Num());
+
+			// Set how many valid bits there are in the new packet
+			Packet.SetNumBits(NewPacketByteCount * 8);
+
+			UE_LOG(PacketHandlerLog, VeryVerbose, TEXT("  AESGCM packet handler sending %" INT64_FMT " bytes after encryption."), Packet.GetNumBytes());
 		}
 		else
 		{
-			NewPacket.SerializeBits(Packet.GetData(), Packet.GetNumBits());
-		}
+			// Make sure packet has space available for extra 8 bits
+			if (Packet.AllowAppend(8))
+			{
+				// Reserve one byte in packet data in the beginning
+				uint8* PacketData = Packet.GetData();
+				FPlatformMemory::Memmove(PacketData + 1, PacketData, Packet.GetNumBytes());
+
+				// First byte 0 means that packet contains unencrypted payload
+				PacketData[0] = 0;
 
-		Packet = MoveTemp(NewPacket);
+				// Include first 8 bits in new packet data
+				Packet.SetNumBits(8 + Packet.GetNumBits());
+			}
+			else
+			{
+				Packet.SetOverflowed(Packet.GetNumBits());
+			}
+		}
 	}
 }
 
-TArray<uint8> FAESGCMHandlerComponent::Decrypt(const TArrayView<const uint8> InCiphertext, const TArrayView<const uint8> IV, const TArrayView<const uint8> AuthTag, EPlatformCryptoResult& DecryptResult)
+EPlatformCryptoResult FAESGCMHandlerComponent::Decrypt(TArrayView<uint8> OutPlaintext, const TArrayView<const uint8> InCiphertext, const TArrayView<const uint8> IV, const TArrayView<const uint8> AuthTag)
 {
-	DecryptResult = EPlatformCryptoResult::Failure;
+	EPlatformCryptoResult Result = Decryptor->Reset(IV);
 
-	TArray<uint8> OutPlaintext;
+	if (Result != EPlatformCryptoResult::Success)
+	{
+		return Result;
+	}
 
-	EPlatformCryptoResult Result = Decryptor->Reset(IV);
-	if (Result == EPlatformCryptoResult::Success)
+	Result = Decryptor->SetAuthTag(AuthTag);
+	if (Result != EPlatformCryptoResult::Success)
 	{
-		Result = Decryptor->SetAuthTag(AuthTag);
-		if (Result == EPlatformCryptoResult::Success)
-		{
-			OutPlaintext.AddUninitialized(Decryptor->GetUpdateBufferSizeBytes(InCiphertext) + Decryptor->GetFinalizeBufferSizeBytes());
+		return Result;
+	}
 
-			int32 UpdateBytesWritten = 0;
-			Result = Decryptor->Update(InCiphertext, OutPlaintext, UpdateBytesWritten);
-			if (Result == EPlatformCryptoResult::Success)
-			{
-				int32 FinalizeBytesWritten = 0;
-				Result = Decryptor->Finalize(TArrayView<uint8>(OutPlaintext.GetData() + UpdateBytesWritten, OutPlaintext.Num() - UpdateBytesWritten), FinalizeBytesWritten);
-				if (Result == EPlatformCryptoResult::Success)
-				{
-					OutPlaintext.SetNum(UpdateBytesWritten + FinalizeBytesWritten);
-					DecryptResult = EPlatformCryptoResult::Success;
-				}
-			}
-		}
+	if (OutPlaintext.Num() < Decryptor->GetUpdateBufferSizeBytes(InCiphertext) + Decryptor->GetFinalizeBufferSizeBytes())
+	{
+		// not enough space in plaintext output array
+		return EPlatformCryptoResult::Failure;
+	}
+
+	int32 UpdateBytesWritten = 0;
+	Result = Decryptor->Update(InCiphertext, OutPlaintext, UpdateBytesWritten);
+	if (Result != EPlatformCryptoResult::Success)
+	{
+		return Result;
+	}
+
+	int32 FinalizeBytesWritten = 0;
+	Result = Decryptor->Finalize(TArrayView<uint8>(OutPlaintext.GetData() + UpdateBytesWritten, OutPlaintext.Num() - UpdateBytesWritten), FinalizeBytesWritten);
+	if (Result != EPlatformCryptoResult::Success)
+	{
+		return Result;
 	}
 
-	return OutPlaintext;
+	if (UpdateBytesWritten + FinalizeBytesWritten != InCiphertext.Num())
+	{
+		// AES GCM mode always decrypts to same amount of bytes as ciphertext
+		return EPlatformCryptoResult::Failure;
+	}
+
+	return EPlatformCryptoResult::Success;
 }
 
-TArray<uint8> FAESGCMHandlerComponent::Encrypt(const TArrayView<const uint8> InPlaintext, const TArrayView<const uint8> IV, TArrayView<uint8> OutAuthTag, EPlatformCryptoResult& EncryptResult)
+EPlatformCryptoResult FAESGCMHandlerComponent::Encrypt(TArrayView<uint8> OutCipherText, const TArrayView<const uint8> InPlaintext, const TArrayView<const uint8> IV, TArrayView<uint8> OutAuthTag)
 {
-	EncryptResult = EPlatformCryptoResult::Failure;
+	EPlatformCryptoResult Result = Encryptor->Reset(IV);
 
-	TArray<uint8> OutCiphertext;
+	if (Result != EPlatformCryptoResult::Success)
+	{
+		return Result;
+	}
 
-	EPlatformCryptoResult Result = Encryptor->Reset(IV);
-	if (Result == EPlatformCryptoResult::Success)
+	if (OutCipherText.Num() < Encryptor->GetUpdateBufferSizeBytes(InPlaintext) + Encryptor->GetFinalizeBufferSizeBytes())
 	{
-		OutCiphertext.AddUninitialized(Encryptor->GetUpdateBufferSizeBytes(InPlaintext) + Encryptor->GetFinalizeBufferSizeBytes());
+		// not enough space in ciphertext output array
+		return EPlatformCryptoResult::Failure;
+	}
 
-		int32 UpdateBytesWritten = 0;
-		Result = Encryptor->Update(InPlaintext, OutCiphertext, UpdateBytesWritten);
-		if (Result == EPlatformCryptoResult::Success)
-		{
-			int32 FinalizeBytesWritten = 0;
-			Result = Encryptor->Finalize(TArrayView<uint8>(OutCiphertext.GetData() + UpdateBytesWritten, OutCiphertext.Num() - UpdateBytesWritten), FinalizeBytesWritten);
-			if (Result == EPlatformCryptoResult::Success)
-			{
-				int32 AuthTagBytesWritten = 0;
-				Result = Encryptor->GenerateAuthTag(OutAuthTag, AuthTagBytesWritten);
-				if (Result == EPlatformCryptoResult::Success)
-				{
-					OutCiphertext.SetNum(UpdateBytesWritten + FinalizeBytesWritten);
-					EncryptResult = EPlatformCryptoResult::Success;
-				}
-			}
-		}
+	int32 UpdateBytesWritten = 0;
+	Result = Encryptor->Update(InPlaintext, OutCipherText, UpdateBytesWritten);
+	if (Result != EPlatformCryptoResult::Success)
+	{
+		return Result;
+	}
+
+	int32 FinalizeBytesWritten = 0;
+	Result = Encryptor->Finalize(TArrayView<uint8>(OutCipherText.GetData() + UpdateBytesWritten, OutCipherText.Num() - UpdateBytesWritten), FinalizeBytesWritten);
+	if (Result != EPlatformCryptoResult::Success)
+	{
+		return Result;
+	}
+
+	int32 AuthTagBytesWritten = 0;
+	Result = Encryptor->GenerateAuthTag(OutAuthTag, AuthTagBytesWritten);
+	if (Result != EPlatformCryptoResult::Success)
+	{
+		return Result;
+	}
+
+	if (UpdateBytesWritten + FinalizeBytesWritten != InPlaintext.Num())
+	{
+		// AES GCM mode always encrypts to same amount of bytes as plaintext
+		return EPlatformCryptoResult::Failure;
 	}
 
-	return OutCiphertext;
+	return EPlatformCryptoResult::Success;
 }
 
 int32 FAESGCMHandlerComponent::GetReservedPacketBits() const
 {
-	// Worst case includes the encryption enabled bit, the termination bit, padding up to the next whole byte, and a block of padding.
-	// Now also includes initialization vector and auth tag
-	return 2 + 7 + ((BlockSizeInBytes + IVSizeInBytes + AuthTagSizeInBytes) * 8);
+	// Worst case includes the IV and AuthTag
+	// For unencrypted packets it is just an extra one byte, which is smaller
+	return (IVSizeInBytes + AuthTagSizeInBytes) * 8;
 }
 
 void FAESGCMHandlerComponent::CountBytes(FArchive& Ar) const
@@ -367,6 +398,4 @@ void FAESGCMHandlerComponent::CountBytes(FArchive& Ar) const
 		LocalContext->CountBytes(Ar);
 	}
 	*/
-
-	Ciphertext.CountBytes(Ar);
 }
diff --git a/Engine/Plugins/Runtime/PacketHandlers/AESGCMHandlerComponent/Source/Public/AESGCMHandlerComponent.h b/Engine/Plugins/Runtime/PacketHandlers/AESGCMHandlerComponent/Source/Public/AESGCMHandlerComponent.h
--- a/Engine/Plugins/Runtime/PacketHandlers/AESGCMHandlerComponent/Source/Public/AESGCMHandlerComponent.h
+++ b/Engine/Plugins/Runtime/PacketHandlers/AESGCMHandlerComponent/Source/Public/AESGCMHandlerComponent.h
@@ -59,9 +59,6 @@ private:
 
 	TUniquePtr<IPlatformCryptoDecryptor> Decryptor;
 	TUniquePtr<IPlatformCryptoEncryptor> Encryptor;
-	
-	// Avoid per packet allocations
-	TArray<uint8> Ciphertext;
 
 	// IV used for encryption
 	TStaticArray<uint8, IVSizeInBytes> OutIV;
@@ -71,8 +68,8 @@ private:
 	/** Fault handler for AESGCM-specific errors, that may trigger NetConnection Close */
 	FAESGCMFaultHandler AESGCMFaultHandler;
 
-	TArray<uint8> Decrypt(const TArrayView<const uint8> InCiphertext, const TArrayView<const uint8> IV, const TArrayView<const uint8> AuthTag, EPlatformCryptoResult& DecryptResult);
-	TArray<uint8> Encrypt(const TArrayView<const uint8> InPlaintext, const TArrayView<const uint8> IV, TArrayView<uint8> OutAuthTag, EPlatformCryptoResult& EncryptResult);
+	EPlatformCryptoResult Decrypt(TArrayView<uint8> OutPlainText, const TArrayView<const uint8> InCiphertext, const TArrayView<const uint8> IV, const TArrayView<const uint8> AuthTag);
+	EPlatformCryptoResult Encrypt(TArrayView<uint8> OutCipherText, const TArrayView<const uint8> InPlaintext, const TArrayView<const uint8> IV, TArrayView<uint8> OutAuthTag);
 };
 
 
diff --git a/Engine/Source/Runtime/Core/Private/Serialization/BitWriter.cpp b/Engine/Source/Runtime/Core/Private/Serialization/BitWriter.cpp
--- a/Engine/Source/Runtime/Core/Private/Serialization/BitWriter.cpp
+++ b/Engine/Source/Runtime/Core/Private/Serialization/BitWriter.cpp
@@ -94,6 +94,12 @@ void FBitWriter::SerializeBits( void* Src, int64 LengthBits )
 				Buffer[Num>>3] |= GShift[Num&7];
 			Num++;
 		}
+		else if (Num == 0)
+		{
+			// Just use memcpy when setting bits directy after Reset()
+			FMemory::Memcpy(Buffer.GetData(), Src, (LengthBits + 7) / 8);
+			Num = LengthBits;
+		}
 		else
 		{
 			appBitsCpy(Buffer.GetData(), Num, (uint8*)Src, 0, LengthBits);
diff --git a/Engine/Source/Runtime/Core/Public/Serialization/BitReader.h b/Engine/Source/Runtime/Core/Public/Serialization/BitReader.h
--- a/Engine/Source/Runtime/Core/Public/Serialization/BitReader.h
+++ b/Engine/Source/Runtime/Core/Public/Serialization/BitReader.h
@@ -203,6 +203,17 @@ CORE_API PRAGMA_ENABLE_BUFFER_OVERRUN_WARNING
 			SetOverflowed(Pos - PrePos);
 		}
 	}
+	FORCEINLINE_DEBUGGABLE void Skip(int32 BitCount)
+	{
+		if (Pos + BitCount > Num)
+		{
+			SetOverflowed(Pos);
+		}
+		else
+		{
+			Pos += BitCount;
+		}
+	}
 
 	/**
 	 * Marks this bit reader as overflowed.
diff --git a/Engine/Source/Runtime/Core/Public/Serialization/BitWriter.h b/Engine/Source/Runtime/Core/Public/Serialization/BitWriter.h
--- a/Engine/Source/Runtime/Core/Public/Serialization/BitWriter.h
+++ b/Engine/Source/Runtime/Core/Public/Serialization/BitWriter.h
@@ -96,6 +96,17 @@ public:
 		return &Buffer;
 	}
 
+	/**
+	 * Sets how many bits packet actually contains
+	 * Must be not greater than max bits allowed
+	 * Useful when writing bytes to buffer manually via pointer from GetData()
+	 */
+	FORCEINLINE void SetNumBits(int64 NumBits)
+	{
+		check(NumBits <= Max);
+		Num = NumBits;
+	}
+
 	/**
 	 * Returns the number of bytes written.
 	 */
diff --git a/Engine/Source/Runtime/Engine/Private/PacketHandlers/StatelessConnectHandlerComponent.cpp b/Engine/Source/Runtime/Engine/Private/PacketHandlers/StatelessConnectHandlerComponent.cpp
--- a/Engine/Source/Runtime/Engine/Private/PacketHandlers/StatelessConnectHandlerComponent.cpp
+++ b/Engine/Source/Runtime/Engine/Private/PacketHandlers/StatelessConnectHandlerComponent.cpp
@@ -1243,7 +1243,19 @@ void StatelessConnectHandlerComponent::Outgoing(FBitWriter& Packet, FOutPacketTr
 {
 	// All UNetConnection packets must specify a zero bHandshakePacket value
 	const int32 AdjustedSize = GetAdjustedSizeBits(Packet.GetNumBits(), LastRemoteHandshakeVersion);
-	FBitWriter NewPacket(AdjustedSize + 1, true);
+
+	// Try to reuse memory from previous Outgoing call if there is enough space
+	FBitWriter NewPacket = MoveTemp(OldPacket);
+	if (NewPacket.GetMaxBits() < AdjustedSize + 1)
+	{
+		// Otherwise allocate at least MAX_PACKET_SIZE bytes, just like PacketHandler does, so packet memory can be reused later
+		NewPacket = FBitWriter(FMath::Max(AdjustedSize + 1, MAX_PACKET_SIZE * 8), true);
+	}
+	else
+	{
+		NewPacket.Reset();
+	}
+
 	uint8 bHandshakePacket = 0;
 
 	if (MagicHeader.Num() > 0)
@@ -1260,6 +1272,9 @@ void StatelessConnectHandlerComponent::Outgoing(FBitWriter& Packet, FOutPacketTr
 	NewPacket.WriteBit(bHandshakePacket);
 	NewPacket.SerializeBits(Packet.GetData(), Packet.GetNumBits());
 
+	// Remember bit writer memory of old packet, so it can be re-used next time
+	OldPacket = MoveTemp(Packet);
+
 	Packet = MoveTemp(NewPacket);
 }
 
diff --git a/Engine/Source/Runtime/Engine/Public/PacketHandlers/StatelessConnectHandlerComponent.h b/Engine/Source/Runtime/Engine/Public/PacketHandlers/StatelessConnectHandlerComponent.h
--- a/Engine/Source/Runtime/Engine/Public/PacketHandlers/StatelessConnectHandlerComponent.h
+++ b/Engine/Source/Runtime/Engine/Public/PacketHandlers/StatelessConnectHandlerComponent.h
@@ -436,6 +436,8 @@ private:
 	/** The net driver associated with this handler - for performing connectionless sends */
 	UNetDriver* Driver;
 
+	/** Cache memory of previously used packets to avoid memory allocations */
+	FBitWriter OldPacket;
 
 	/** Serverside variables */
 
