Increase output buffer in AES GCM packet handler by two AES blocks
Because platform crypto might have checks to expect bigger output buffer - but it will never write more encrypted/decrypted data. This will be fixed later.

#rnx

diff --git a/Engine/Plugins/Runtime/PacketHandlers/AESGCMHandlerComponent/Source/Private/AESGCMHandlerComponent.cpp b/Engine/Plugins/Runtime/PacketHandlers/AESGCMHandlerComponent/Source/Private/AESGCMHandlerComponent.cpp
--- a/Engine/Plugins/Runtime/PacketHandlers/AESGCMHandlerComponent/Source/Private/AESGCMHandlerComponent.cpp
+++ b/Engine/Plugins/Runtime/PacketHandlers/AESGCMHandlerComponent/Source/Private/AESGCMHandlerComponent.cpp
@@ -177,7 +177,8 @@ void FAESGCMHandlerComponent::Incoming(FIncomingPacketRef PacketRef)
 
 			UE_LOG(PacketHandlerLog, VeryVerbose, TEXT("AESGCM packet handler received %d bytes before decryption."), CipherText.Num());
 
-			uint8 PlainText[MAX_PACKET_SIZE];
+			// Temporary increase by two blocks because platfrom crypto module might want bigger buffer (but it will never actually write more data)
+			uint8 PlainText[MAX_PACKET_SIZE + 2*BlockSizeInBytes];
 
 			// Decrypt payload and verify AuthTag
 			EPlatformCryptoResult DecryptResult = Decrypt(PlainText, CipherText, IV, AuthTag);
@@ -240,11 +241,13 @@ void FAESGCMHandlerComponent::Outgoing(FBitWriter& Packet, FOutPacketTraits& Tra
 			 
 			// Make sure there is enough space allocated for outgoing packet memory
 			int64 NewPacketByteCount = IVSizeInBytes + AuthTagSizeInBytes + PlainText.Num();
-			if (NewPacketByteCount * 8 > Packet.GetMaxBits())
+			if ((NewPacketByteCount + 2*BlockSizeInBytes) * 8 > Packet.GetMaxBits())
 			{
 				// Allocate max MAX_PACKET_SIZE bytes, just like PacketHandler does, so packet memory can be reused later
 				check(NewPacketByteCount <= MAX_PACKET_SIZE);
-				Packet = FBitWriter(MAX_PACKET_SIZE * 8);
+
+				// Temporary increase by two blocks because platfrom crypto migh want bigger buffer (but it will never actually write more data)
+				Packet = FBitWriter((MAX_PACKET_SIZE + 2*BlockSizeInBytes) * 8);
 			}
 
 			uint8* NewPacketData = Packet.GetData();
