Update change for previous fix - correctly increase buffer for outgoing packets in AES GCM packet handler

#rnx

diff --git a/Engine/Plugins/Runtime/PacketHandlers/AESGCMHandlerComponent/Source/Private/AESGCMHandlerComponent.cpp b/Engine/Plugins/Runtime/PacketHandlers/AESGCMHandlerComponent/Source/Private/AESGCMHandlerComponent.cpp
--- a/Engine/Plugins/Runtime/PacketHandlers/AESGCMHandlerComponent/Source/Private/AESGCMHandlerComponent.cpp
+++ b/Engine/Plugins/Runtime/PacketHandlers/AESGCMHandlerComponent/Source/Private/AESGCMHandlerComponent.cpp
@@ -1,11 +1,18 @@
 // Copyright Epic Games, Inc. All Rights Reserved.
 
 #include "AESGCMHandlerComponent.h"
+#include "HAL/ConsoleManager.h"
 
 #include "PlatformCryptoTypes.h"
 
 IMPLEMENT_MODULE( FAESGCMHandlerComponentModule, AESGCMHandlerComponent )
 
+static int32 GPacketHandlerLegacyBufferSize = 0;
+
+static FAutoConsoleVariableRef CVarPacketHandlerLegacyBufferSize(
+	TEXT("net.PacketHandlerLegacyBufferSize"),
+	GPacketHandlerLegacyBufferSize,
+	TEXT("Enables legacy packet handler buffer size for backwards compatibility"));
 
 const int32 FAESGCMHandlerComponent::KeySizeInBytes;
 const int32 FAESGCMHandlerComponent::BlockSizeInBytes;
@@ -177,22 +184,43 @@ void FAESGCMHandlerComponent::Incoming(FIncomingPacketRef PacketRef)
 
 			UE_LOG(PacketHandlerLog, VeryVerbose, TEXT("AESGCM packet handler received %d bytes before decryption."), CipherText.Num());
 
-			// Temporary increase by two blocks because platfrom crypto module might want bigger buffer (but it will never actually write more data)
-			uint8 PlainText[MAX_PACKET_SIZE + 2*BlockSizeInBytes];
-
-			// Decrypt payload and verify AuthTag
-			EPlatformCryptoResult DecryptResult = Decrypt(PlainText, CipherText, IV, AuthTag);
-			if (DecryptResult == EPlatformCryptoResult::Failure)
+			if (GPacketHandlerLegacyBufferSize)
 			{
-				UE_LOG(PacketHandlerLog, Log, TEXT("FAESGCMHandlerComponent::Incoming: failed to decrypt packet."));
+				uint8 PlainText[MAX_PACKET_SIZE];
 
-				Packet.SetError();
-				AddToChainResultPtr(Traits.ExtendedError, EAESGCMNetResult::AESDecryptionFailed);
+				// Decrypt payload and verify AuthTag
+				EPlatformCryptoResult DecryptResult = Decrypt(PlainText, CipherText, IV, AuthTag);
+				if (DecryptResult == EPlatformCryptoResult::Failure)
+				{
+					UE_LOG(PacketHandlerLog, Log, TEXT("FAESGCMHandlerComponent::Incoming: failed to decrypt packet."));
 
-				return;
+					Packet.SetError();
+					AddToChainResultPtr(Traits.ExtendedError, EAESGCMNetResult::AESDecryptionFailed);
+
+					return;
+				}
+
+				Packet.SetData(PlainText, CipherText.Num() * 8);
 			}
+			else
+			{
+				// Temporary increase by two blocks because platfrom crypto module might want bigger buffer (but it will never actually write more data)
+				uint8 PlainText[MAX_PACKET_SIZE + 2 * BlockSizeInBytes];
+
+				// Decrypt payload and verify AuthTag
+				EPlatformCryptoResult DecryptResult = Decrypt(PlainText, CipherText, IV, AuthTag);
+				if (DecryptResult == EPlatformCryptoResult::Failure)
+				{
+					UE_LOG(PacketHandlerLog, Log, TEXT("FAESGCMHandlerComponent::Incoming: failed to decrypt packet."));
+
+					Packet.SetError();
+					AddToChainResultPtr(Traits.ExtendedError, EAESGCMNetResult::AESDecryptionFailed);
 
-			Packet.SetData(PlainText, CipherText.Num() * 8);
+					return;
+				}
+
+				Packet.SetData(PlainText, CipherText.Num() * 8);
+			}
 		}
 		else
 		{
@@ -228,42 +256,82 @@ void FAESGCMHandlerComponent::Outgoing(FBitWriter& Packet, FOutPacketTraits& Tra
 			}
 
 			TArrayView<uint8> PlainText = { Packet.GetData(), static_cast<int32>(Packet.GetNumBytes()) };
-			uint8 CipherText[MAX_PACKET_SIZE];
 
-			// Encrypt payload & write AuthTag
-			EPlatformCryptoResult EncryptResult = Encrypt(CipherText, PlainText, OutIV, AuthTag);
-			if (EncryptResult == EPlatformCryptoResult::Failure)
+			if (GPacketHandlerLegacyBufferSize)
 			{
-				UE_LOG(PacketHandlerLog, Log, TEXT("FAESGCMHandlerComponent::Outgoing: failed to encrypt packet."));
-				Packet.SetError();
-				return;
+				// Temporary increase by two blocks because platfrom crypto module might want bigger buffer (but it will never actually write more data)
+				uint8 CipherText[MAX_PACKET_SIZE];
+
+				// Encrypt payload & write AuthTag
+				EPlatformCryptoResult EncryptResult = Encrypt(CipherText, PlainText, OutIV, AuthTag);
+				if (EncryptResult == EPlatformCryptoResult::Failure)
+				{
+					UE_LOG(PacketHandlerLog, Log, TEXT("FAESGCMHandlerComponent::Outgoing: failed to encrypt packet."));
+					Packet.SetError();
+					return;
+				}
+
+				// Make sure there is enough space allocated for outgoing packet memory
+				int64 NewPacketByteCount = IVSizeInBytes + AuthTagSizeInBytes + PlainText.Num();
+				if (NewPacketByteCount * 8 > Packet.GetMaxBits())
+				{
+					// Allocate max MAX_PACKET_SIZE bytes, just like PacketHandler does, so packet memory can be reused later
+					check(NewPacketByteCount <= MAX_PACKET_SIZE);
+					Packet = FBitWriter(MAX_PACKET_SIZE * 8);
+				}
+
+				uint8* NewPacketData = Packet.GetData();
+
+				// Copy the IV, AuthTag and encrypted payload to new packet
+				FPlatformMemory::Memcpy(NewPacketData, OutIV.GetData(), OutIV.Num());
+				NewPacketData += OutIV.Num();
+
+				FPlatformMemory::Memcpy(NewPacketData, AuthTag.GetData(), AuthTag.Num());
+				NewPacketData += AuthTag.Num();
+
+				FPlatformMemory::Memcpy(NewPacketData, CipherText, PlainText.Num());
+
+				// Set how many valid bits there are in the new packet
+				Packet.SetNumBits(NewPacketByteCount * 8);
 			}
-			 
-			// Make sure there is enough space allocated for outgoing packet memory
-			int64 NewPacketByteCount = IVSizeInBytes + AuthTagSizeInBytes + PlainText.Num();
-			if ((NewPacketByteCount + 2*BlockSizeInBytes) * 8 > Packet.GetMaxBits())
+			else
 			{
-				// Allocate max MAX_PACKET_SIZE bytes, just like PacketHandler does, so packet memory can be reused later
-				check(NewPacketByteCount <= MAX_PACKET_SIZE);
-
-				// Temporary increase by two blocks because platfrom crypto migh want bigger buffer (but it will never actually write more data)
-				Packet = FBitWriter((MAX_PACKET_SIZE + 2*BlockSizeInBytes) * 8);
+				// Temporary increase by two blocks because platfrom crypto module might want bigger buffer (but it will never actually write more data)
+				uint8 CipherText[MAX_PACKET_SIZE + 2 * BlockSizeInBytes];
+
+				// Encrypt payload & write AuthTag
+				EPlatformCryptoResult EncryptResult = Encrypt(CipherText, PlainText, OutIV, AuthTag);
+				if (EncryptResult == EPlatformCryptoResult::Failure)
+				{
+					UE_LOG(PacketHandlerLog, Log, TEXT("FAESGCMHandlerComponent::Outgoing: failed to encrypt packet."));
+					Packet.SetError();
+					return;
+				}
+
+				// Make sure there is enough space allocated for outgoing packet memory
+				int64 NewPacketByteCount = IVSizeInBytes + AuthTagSizeInBytes + PlainText.Num();
+				if (NewPacketByteCount * 8 > Packet.GetMaxBits())
+				{
+					// Allocate max MAX_PACKET_SIZE bytes, just like PacketHandler does, so packet memory can be reused later
+					check(NewPacketByteCount <= MAX_PACKET_SIZE);
+					Packet = FBitWriter(MAX_PACKET_SIZE * 8);
+				}
+
+				uint8* NewPacketData = Packet.GetData();
+
+				// Copy the IV, AuthTag and encrypted payload to new packet
+				FPlatformMemory::Memcpy(NewPacketData, OutIV.GetData(), OutIV.Num());
+				NewPacketData += OutIV.Num();
+
+				FPlatformMemory::Memcpy(NewPacketData, AuthTag.GetData(), AuthTag.Num());
+				NewPacketData += AuthTag.Num();
+
+				FPlatformMemory::Memcpy(NewPacketData, CipherText, PlainText.Num());
+
+				// Set how many valid bits there are in the new packet
+				Packet.SetNumBits(NewPacketByteCount * 8);
 			}
 
-			uint8* NewPacketData = Packet.GetData();
-
-			// Copy the IV, AuthTag and encrypted payload to new packet
-			FPlatformMemory::Memcpy(NewPacketData, OutIV.GetData(), OutIV.Num());
-			NewPacketData += OutIV.Num();
-
-			FPlatformMemory::Memcpy(NewPacketData, AuthTag.GetData(), AuthTag.Num());
-			NewPacketData += AuthTag.Num();
-			
-			FPlatformMemory::Memcpy(NewPacketData, CipherText, PlainText.Num());
-
-			// Set how many valid bits there are in the new packet
-			Packet.SetNumBits(NewPacketByteCount * 8);
-
 			UE_LOG(PacketHandlerLog, VeryVerbose, TEXT("  AESGCM packet handler sending %" INT64_FMT " bytes after encryption."), Packet.GetNumBytes());
 		}
 		else
