Handle packet bit-lengths that are not multiple of eight in AESGCMHandlerComponent

In case AESGCM component is used on packets that have bit-length not in exact byte size then it will pad them to next byte boundary due to how AES works. This must be accounted in reserved packet bit count, otherwise packet data serializer in other layers can try to put it too many bits overflowing max packet size.

diff --git a/Engine/Plugins/Runtime/PacketHandlers/AESGCMHandlerComponent/Source/Private/AESGCMHandlerComponent.cpp b/Engine/Plugins/Runtime/PacketHandlers/AESGCMHandlerComponent/Source/Private/AESGCMHandlerComponent.cpp
--- a/Engine/Plugins/Runtime/PacketHandlers/AESGCMHandlerComponent/Source/Private/AESGCMHandlerComponent.cpp
+++ b/Engine/Plugins/Runtime/PacketHandlers/AESGCMHandlerComponent/Source/Private/AESGCMHandlerComponent.cpp
@@ -88,7 +88,7 @@ bool FAESGCMHandlerComponent::IsValid() const
 }
 
 
-// Incoming & Outgoing pakcets will always operate on byte level. No bit-aligned stuff is supported here for efficiency reasons.
+// Incoming & Outgoing packets will always operate on byte level. No bit-aligned stuff is supported here for efficiency reasons.
 
 // Encrypted packed layout: 
 // [iv:12] [auth:16] [ciphertext:N]
@@ -376,9 +376,9 @@ EPlatformCryptoResult FAESGCMHandlerComponent::Encrypt(TArrayView<uint8> OutCiph
 
 int32 FAESGCMHandlerComponent::GetReservedPacketBits() const
 {
-	// Worst case includes the IV and AuthTag
-	// For unencrypted packets it is just an extra one byte, which is smaller
-	return (IVSizeInBytes + AuthTagSizeInBytes) * 8;
+	// Worst case includes padding up to the next whole byte and the IV and AuthTag
+	// For unencrypted packets it is padding plus an extra one byte, which is smaller
+	return 7 + (IVSizeInBytes + AuthTagSizeInBytes) * 8;
 }
 
 void FAESGCMHandlerComponent::CountBytes(FArchive& Ar) const
